Add experimental ImageWidget plugin with SImageViewport widget

diff --git a/Engine/Plugins/Experimental/ImageWidgets/ImageWidgets.uplugin b/Engine/Plugins/Experimental/ImageWidgets/ImageWidgets.uplugin
new file mode 100644
--- /dev/null
+++ b/Engine/Plugins/Experimental/ImageWidgets/ImageWidgets.uplugin
@@ -0,0 +1,17 @@
+﻿{
+  "FileVersion": 3,
+  "Version": 1,
+  "VersionName": "0.1",
+  "Description": "Generic Slate widgets for displaying images and image-like content, and content related to images.",
+  "Category": "Other",
+  "CreatedBy": "Epic Games, Inc.",
+  "CreatedByURL": "https://epicgames.com",
+  "IsBetaVersion": false,
+  "IsExperimentalVersion": true,
+  "Modules": [
+    {
+      "Name": "ImageWidgets",
+      "Type": "Editor"
+    }
+  ]
+}
diff --git a/Engine/Plugins/Experimental/ImageWidgets/Source/ImageWidgets/ImageWidgets.Build.cs b/Engine/Plugins/Experimental/ImageWidgets/Source/ImageWidgets/ImageWidgets.Build.cs
new file mode 100644
--- /dev/null
+++ b/Engine/Plugins/Experimental/ImageWidgets/Source/ImageWidgets/ImageWidgets.Build.cs
@@ -0,0 +1,56 @@
+﻿// Copyright Epic Games, Inc. All Rights Reserved.
+
+using UnrealBuildTool;
+
+public class ImageWidgets : ModuleRules
+{
+	// Set this flag to true to force building the ColorViewer sample. Otherwise it will only be built in Debug builds. 
+	private const bool ForceBuildColorViewerSample = false;
+
+	// Todo Remove build flags as soon as all the respective prototype code is migrated to this module.
+	// These flags are currently used to migrate already existing prototype functionality into this module without having to remove/change the functionality of
+	// the prototype applications.
+	public static bool EnableAbComparison = false;
+	public static bool EnableCatalog = false;
+
+	public ImageWidgets(ReadOnlyTargetRules Target) : base(Target)
+	{
+		PCHUsage = PCHUsageMode.UseExplicitOrSharedPCHs;
+
+		PublicDependencyModuleNames.AddRange(
+			new[]
+			{
+				"Core",
+				"Slate",
+				"UnrealEd"
+			}
+		);
+
+		PrivateDependencyModuleNames.AddRange(
+			new[]
+			{
+				"ApplicationCore",
+				"CoreUObject",
+				"Engine",
+				"InputCore",
+				"SlateCore"
+			}
+		);
+
+		bool bBuildColorViewerSample = ForceBuildColorViewerSample || Target.Configuration == UnrealTargetConfiguration.Debug;
+		if (bBuildColorViewerSample)
+		{
+			PrivateDependencyModuleNames.AddRange(
+				new[]
+				{
+					"Projects",
+					"WorkspaceMenuStructure"
+				}
+			);
+		}
+		PublicDefinitions.Add(bBuildColorViewerSample ? "IMAGE_WIDGETS_BUILD_COLOR_VIEWER_SAMPLE=1" : "IMAGE_WIDGETS_BUILD_COLOR_VIEWER_SAMPLE=0");
+
+		PublicDefinitions.Add(EnableAbComparison ? "IMAGE_WIDGETS_WITH_AB_COMPARISON=1" : "IMAGE_WIDGETS_WITH_AB_COMPARISON=0");
+		PublicDefinitions.Add(EnableCatalog ? "IMAGE_WIDGETS_WITH_CATALOG=1" : "IMAGE_WIDGETS_WITH_CATALOG=0");
+	}
+}
\ No newline at end of file
diff --git a/Engine/Plugins/Experimental/ImageWidgets/Source/ImageWidgets/Private/ColorViewerSample/ColorViewer.cpp b/Engine/Plugins/Experimental/ImageWidgets/Source/ImageWidgets/Private/ColorViewerSample/ColorViewer.cpp
new file mode 100644
--- /dev/null
+++ b/Engine/Plugins/Experimental/ImageWidgets/Source/ImageWidgets/Private/ColorViewerSample/ColorViewer.cpp
@@ -0,0 +1,65 @@
+﻿// Copyright Epic Games, Inc. All Rights Reserved.
+
+#if IMAGE_WIDGETS_BUILD_COLOR_VIEWER_SAMPLE
+
+#include "ColorViewer.h"
+
+#include <CanvasItem.h>
+#include <CanvasTypes.h>
+
+#define LOCTEXT_NAMESPACE "ColorViewer"
+
+namespace UE::ImageWidgets::Sample
+{
+	IImageViewer::FImageInfo FColorViewer::GetCurrentImageInfo() const
+	{
+		return {ImageGuid, ImageSize, 0, bColorIsValid};
+	}
+
+	void FColorViewer::DrawCurrentImage(FViewport* Viewport, FCanvas* Canvas, const FDrawProperties& Properties)
+	{
+		// Get color value after tone mapping.
+		const FLinearColor ToneMappedColor = ToneMapping.GetToneMappedColor(Color);
+
+		// Draw simple quad with current tone mapped color.
+		// In a less trivial use case, this would require rendering quads with textures and the like. 
+		FCanvasTileItem Tile(Properties.Placement.Offset, Properties.Placement.Size, ToneMappedColor);
+		Canvas->DrawItem(Tile);
+	}
+
+	TOptional<TVariant<FColor, FLinearColor>> FColorViewer::GetCurrentImagePixelColor(FIntPoint PixelCoords, int32 MipLevel) const
+	{
+		if (bColorIsValid)
+		{
+			// Returns the current color as float values.
+			// In a less trivial use case, the pixel coordinates and potentially the MIP level would be needed to look up the color value.
+			return TVariant<FColor, FLinearColor>(TInPlaceType<FLinearColor>(), Color);
+		}
+		return {};
+	}
+
+	void FColorViewer::RandomizeColor()
+	{
+		auto Random = []
+		{
+			return FMath::RandRange(0.0f, 1.0f);
+		};
+
+		Color = {Random(), Random(), Random()};
+		bColorIsValid = true;
+	}
+
+	FToneMapping::EMode FColorViewer::GetToneMapping() const
+	{
+		return ToneMapping.Mode;
+	}
+
+	void FColorViewer::SetToneMapping(FToneMapping::EMode Mode)
+	{
+		ToneMapping.Mode = Mode;
+	}
+}
+
+#undef LOCTEXT_NAMESPACE
+
+#endif // IMAGE_WIDGETS_BUILD_COLOR_VIEWER_SAMPLE
\ No newline at end of file
diff --git a/Engine/Plugins/Experimental/ImageWidgets/Source/ImageWidgets/Private/ColorViewerSample/ColorViewer.h b/Engine/Plugins/Experimental/ImageWidgets/Source/ImageWidgets/Private/ColorViewerSample/ColorViewer.h
new file mode 100644
--- /dev/null
+++ b/Engine/Plugins/Experimental/ImageWidgets/Source/ImageWidgets/Private/ColorViewerSample/ColorViewer.h
@@ -0,0 +1,74 @@
+﻿// Copyright Epic Games, Inc. All Rights Reserved.
+
+#pragma once
+
+#if IMAGE_WIDGETS_BUILD_COLOR_VIEWER_SAMPLE
+
+#include <IImageViewer.h>
+
+namespace UE::ImageWidgets::Sample
+{
+	/**
+	 * Provide tone mapping capabilities.
+	 * In this simple example, this is limited to just normal RGB plus luminance (grayscale).
+	 */
+	struct FToneMapping
+	{
+		enum class EMode { RGB, Lum };
+
+		FToneMapping(EMode Mode)
+			: Mode(Mode)
+		{
+		}
+
+		FLinearColor GetToneMappedColor(const FLinearColor& Color) const
+		{
+			if (Mode == EMode::Lum)
+			{
+				const float Luminance = 0.3f * Color.R + 0.59f * Color.G + 0.11f * Color.B;
+				return FLinearColor(Luminance, Luminance, Luminance);
+			}
+			return Color;
+		}
+
+		EMode Mode;
+	};
+
+	/**
+	 * Image viewer implementation used by the image widgets.
+	 * It contains any image data, in this case just colors, and renders the image data in the viewport widgets based on the viewport provided parameters.
+	 */
+	class FColorViewer final : public IImageViewer
+	{
+	public:
+		// IImageViewer overrides - begin
+		virtual FImageInfo GetCurrentImageInfo() const override;
+		virtual void DrawCurrentImage(FViewport* Viewport, FCanvas* Canvas, const FDrawProperties& Properties) override;
+		virtual TOptional<TVariant<FColor, FLinearColor>> GetCurrentImagePixelColor(FIntPoint PixelCoords, int32 MipLevel) const override;
+		// IImageViewer overrides - end
+
+		/** Sets a random color as the "image". This is a simple proxy for the image content changing and/or users choosing different images to display. */
+		void RandomizeColor();
+
+		/** Access to tone mapping data. This is effectively used by the viewport toolbar extensions as well as when drawing the image. */
+		FToneMapping::EMode GetToneMapping() const;
+		void SetToneMapping(FToneMapping::EMode Mode);
+
+	private:
+		/** The tone mapping data. */
+		FToneMapping ToneMapping = FToneMapping::EMode::RGB;
+
+		/** The color for the currently displayed image. */
+		FLinearColor Color;
+
+		/** This is set to true once the color is initialized by the user, and evaluated to tell the viewport if the image is valid. */
+		bool bColorIsValid = false;
+
+		/** Hardcoded values for the image size and the GUID of the single image this is currently able to display.
+		 *  In a more realistic application, these value would depend on the actual images that are available. */
+		inline static constexpr FGuid ImageGuid = FGuid(1, 0, 0, 0);
+		inline static const FIntPoint ImageSize = 512;
+	};
+}
+
+#endif // IMAGE_WIDGETS_BUILD_COLOR_VIEWER_SAMPLE
diff --git a/Engine/Plugins/Experimental/ImageWidgets/Source/ImageWidgets/Private/ColorViewerSample/ColorViewerCommands.cpp b/Engine/Plugins/Experimental/ImageWidgets/Source/ImageWidgets/Private/ColorViewerSample/ColorViewerCommands.cpp
new file mode 100644
--- /dev/null
+++ b/Engine/Plugins/Experimental/ImageWidgets/Source/ImageWidgets/Private/ColorViewerSample/ColorViewerCommands.cpp
@@ -0,0 +1,30 @@
+﻿// Copyright Epic Games, Inc. All Rights Reserved.
+
+#if IMAGE_WIDGETS_BUILD_COLOR_VIEWER_SAMPLE
+
+#include "ColorViewerCommands.h"
+
+#include "ColorViewerStyle.h"
+#include "Framework/Commands/InputChord.h"
+
+#define LOCTEXT_NAMESPACE "ColorViewerCommands"
+
+namespace UE::ImageWidgets::Sample
+{
+	FColorViewerCommands::FColorViewerCommands()
+		: TCommands("ColorViewer", LOCTEXT("ContextDescription", "Color Viewer"), NAME_None, FColorViewerStyle::Get().GetStyleSetName())
+	{
+	}
+
+	void FColorViewerCommands::RegisterCommands()
+	{
+		UI_COMMAND(RandomizeColor, "Randomize", "Set random color", EUserInterfaceActionType::Button, FInputChord(EKeys::R));
+
+		UI_COMMAND(ToneMappingRGB, "RGB", "Show full color", EUserInterfaceActionType::ToggleButton, FInputChord());
+		UI_COMMAND(ToneMappingLum, "Lum", "Show luminance only", EUserInterfaceActionType::ToggleButton, FInputChord());
+	}
+}
+
+#undef LOCTEXT_NAMESPACE
+
+#endif // IMAGE_WIDGETS_BUILD_COLOR_VIEWER_SAMPLE
diff --git a/Engine/Plugins/Experimental/ImageWidgets/Source/ImageWidgets/Private/ColorViewerSample/ColorViewerCommands.h b/Engine/Plugins/Experimental/ImageWidgets/Source/ImageWidgets/Private/ColorViewerSample/ColorViewerCommands.h
new file mode 100644
--- /dev/null
+++ b/Engine/Plugins/Experimental/ImageWidgets/Source/ImageWidgets/Private/ColorViewerSample/ColorViewerCommands.h
@@ -0,0 +1,28 @@
+﻿// Copyright Epic Games, Inc. All Rights Reserved.
+
+#pragma once
+
+#if IMAGE_WIDGETS_BUILD_COLOR_VIEWER_SAMPLE
+
+#include "Framework/Commands/Commands.h"
+#include "Framework/Commands/UICommandInfo.h"
+
+namespace UE::ImageWidgets::Sample
+{
+	/**
+	 * Provides commands used for the color viewer sample.
+	 */
+	class FColorViewerCommands : public TCommands<FColorViewerCommands>
+	{
+	public:
+		FColorViewerCommands();
+
+		virtual void RegisterCommands() override;
+
+		TSharedPtr<FUICommandInfo> RandomizeColor; // Button and hotkey for setting a random color
+		TSharedPtr<FUICommandInfo> ToneMappingRGB; // Toggle button for switching to RGB (no tone mapping)
+		TSharedPtr<FUICommandInfo> ToneMappingLum; // Toggle button for switching to luminance (grayscale tone mapping)
+	};
+}
+
+#endif // IMAGE_WIDGETS_BUILD_COLOR_VIEWER_SAMPLE
diff --git a/Engine/Plugins/Experimental/ImageWidgets/Source/ImageWidgets/Private/ColorViewerSample/ColorViewerStyle.cpp b/Engine/Plugins/Experimental/ImageWidgets/Source/ImageWidgets/Private/ColorViewerSample/ColorViewerStyle.cpp
new file mode 100644
--- /dev/null
+++ b/Engine/Plugins/Experimental/ImageWidgets/Source/ImageWidgets/Private/ColorViewerSample/ColorViewerStyle.cpp
@@ -0,0 +1,43 @@
+﻿// Copyright Epic Games, Inc. All Rights Reserved.
+
+#if IMAGE_WIDGETS_BUILD_COLOR_VIEWER_SAMPLE
+
+#include "ColorViewerStyle.h"
+
+#include "Interfaces/IPluginManager.h"
+#include "Styling/AppStyle.h"
+#include "Styling/CoreStyle.h"
+#include "Styling/SlateStyleRegistry.h"
+#include "Styling/SlateTypes.h"
+
+namespace UE::ImageWidgets::Sample
+{
+	FName FColorViewerStyle::StyleName("ColorViewerStyle");
+
+	FColorViewerStyle::FColorViewerStyle()
+		: FSlateStyleSet(StyleName)
+	{
+		const FString ContentDir = IPluginManager::Get().FindPlugin(TEXT("ImageWidgets"))->GetContentDir();
+
+		const FVector2f IconSizeToolbar(16.0f, 16.0f);
+
+		// Use custom icons in the content folder of this plugin.
+		Set("ToneMappingRGB", new FSlateImageBrush(ContentDir / TEXT("Icons/ToneMappingRGB") + TEXT(".png"), IconSizeToolbar));
+		Set("ToneMappingLum", new FSlateImageBrush(ContentDir / TEXT("Icons/ToneMappingLum") + TEXT(".png"), IconSizeToolbar));
+
+		FSlateStyleRegistry::RegisterSlateStyle(*this);
+	}
+
+	FColorViewerStyle::~FColorViewerStyle()
+	{
+		FSlateStyleRegistry::UnRegisterSlateStyle(*this);
+	}
+
+	FColorViewerStyle& FColorViewerStyle::Get()
+	{
+		static FColorViewerStyle Instance;
+		return Instance;
+	}
+}
+
+#endif // IMAGE_WIDGETS_BUILD_COLOR_VIEWER_SAMPLE
diff --git a/Engine/Plugins/Experimental/ImageWidgets/Source/ImageWidgets/Private/ColorViewerSample/ColorViewerStyle.h b/Engine/Plugins/Experimental/ImageWidgets/Source/ImageWidgets/Private/ColorViewerSample/ColorViewerStyle.h
new file mode 100644
--- /dev/null
+++ b/Engine/Plugins/Experimental/ImageWidgets/Source/ImageWidgets/Private/ColorViewerSample/ColorViewerStyle.h
@@ -0,0 +1,27 @@
+﻿// Copyright Epic Games, Inc. All Rights Reserved.
+
+#pragma once
+
+#if IMAGE_WIDGETS_BUILD_COLOR_VIEWER_SAMPLE
+
+#include "Styling/SlateStyle.h"
+
+namespace UE::ImageWidgets::Sample
+{
+	/**
+	 * Style declarations for the color viewer sample.
+	 */
+	class FColorViewerStyle : public FSlateStyleSet
+	{
+	public:
+		static FName StyleName;
+
+		static FColorViewerStyle& Get();
+
+	private:
+		FColorViewerStyle();
+		virtual ~FColorViewerStyle() override;
+	};
+}
+
+#endif // IMAGE_WIDGETS_BUILD_COLOR_VIEWER_SAMPLE
\ No newline at end of file
diff --git a/Engine/Plugins/Experimental/ImageWidgets/Source/ImageWidgets/Private/ColorViewerSample/ColorViewerWidget.cpp b/Engine/Plugins/Experimental/ImageWidgets/Source/ImageWidgets/Private/ColorViewerSample/ColorViewerWidget.cpp
new file mode 100644
--- /dev/null
+++ b/Engine/Plugins/Experimental/ImageWidgets/Source/ImageWidgets/Private/ColorViewerSample/ColorViewerWidget.cpp
@@ -0,0 +1,118 @@
+﻿// Copyright Epic Games, Inc. All Rights Reserved.
+
+#if IMAGE_WIDGETS_BUILD_COLOR_VIEWER_SAMPLE
+
+#include "ColorViewerWidget.h"
+
+#include "ColorViewerCommands.h"
+#include "ColorViewerStyle.h"
+#include "Framework/MultiBox/MultiBoxBuilder.h"
+
+#define LOCTEXT_NAMESPACE "ColorViewerWidget"
+
+namespace UE::ImageWidgets::Sample
+{
+	void SColorViewerWidget::Construct(const FArguments&)
+	{
+		ColorViewer = MakeShared<FColorViewer>();
+
+		BindCommands();
+
+		// Create toolbar extensions for a button to randomize the displayed color as well as the tone mapping controls.
+		const TSharedPtr<FExtender> ToolbarExtender = MakeShared<FExtender>();
+		ToolbarExtender->AddToolBarExtension("ToolbarCenter", EExtensionHook::Before, CommandList,
+		                                     FToolBarExtensionDelegate::CreateSP(this, &SColorViewerWidget::AddRandomColorButton));
+		ToolbarExtender->AddToolBarExtension("ToolbarRight", EExtensionHook::After, CommandList,
+		                                     FToolBarExtensionDelegate::CreateSP(this, &SColorViewerWidget::AddToneMappingButtons));
+
+		// Fill the widget with the image viewport.
+		ChildSlot
+			.VAlign(VAlign_Fill)
+			.HAlign(HAlign_Fill)
+			[
+				SAssignNew(Viewport, ImageWidgets::SImageViewport, ColorViewer.ToSharedRef())
+					.ToolbarExtender(ToolbarExtender)
+					.DrawSettings(SImageViewport::FDrawSettings{
+							.ClearColor = FLinearColor::Black,
+							.bBorderEnabled = true,
+							.BorderThickness = 1.0,
+							.BorderColor = FVector3f(0.2f),
+							.bBackgroundColorEnabled = false,
+							.bBackgroundCheckerEnabled = false
+						})
+			];
+	}
+
+	FReply SColorViewerWidget::OnKeyDown(const FGeometry& Geometry, const FKeyEvent& KeyEvent)
+	{
+		// Capture all key binds that are handled by the widget's commands. 
+		if (CommandList->ProcessCommandBindings(KeyEvent))
+		{
+			return FReply::Handled();
+		}
+
+		return FReply::Unhandled();
+	}
+
+	void SColorViewerWidget::AddRandomColorButton(FToolBarBuilder& ToolbarBuilder) const
+	{
+		const FSlateIcon RandomizeColorIcon(FAppStyle::GetAppStyleSetName(), "FontEditor.Update");
+
+		ToolbarBuilder.AddToolBarButton(FColorViewerCommands::Get().RandomizeColor, NAME_None, TAttribute<FText>(), TAttribute<FText>(),
+		                                TAttribute<FSlateIcon>(RandomizeColorIcon));
+	}
+
+	void SColorViewerWidget::AddToneMappingButtons(FToolBarBuilder& ToolbarBuilder) const
+	{
+		ToolbarBuilder.AddSeparator();
+
+		ToolbarBuilder.BeginBlockGroup();
+		{
+			const FName& StyleSetName = FColorViewerStyle::Get().GetStyleSetName();
+			const FColorViewerCommands& Commands = FColorViewerCommands::Get();
+
+			const FSlateIcon RGBIcon(StyleSetName, "ToneMappingRGB");
+			const FSlateIcon LumIcon(StyleSetName, "ToneMappingLum");
+
+			ToolbarBuilder.AddToolBarButton(Commands.ToneMappingRGB, NAME_None, TAttribute<FText>(), TAttribute<FText>(), TAttribute<FSlateIcon>(RGBIcon));
+			ToolbarBuilder.AddToolBarButton(Commands.ToneMappingLum, NAME_None, TAttribute<FText>(), TAttribute<FText>(), TAttribute<FSlateIcon>(LumIcon));
+		}
+		ToolbarBuilder.EndBlockGroup();
+	}
+
+	void SColorViewerWidget::BindCommands()
+	{
+		const FColorViewerCommands& Commands = FColorViewerCommands::Get();
+
+		CommandList = MakeShared<FUICommandList>();
+
+		CommandList->MapAction(
+			Commands.RandomizeColor,
+			FExecuteAction::CreateSP(ColorViewer.ToSharedRef(), &FColorViewer::RandomizeColor)
+		);
+
+		CommandList->MapAction(
+			Commands.ToneMappingRGB,
+			FExecuteAction::CreateSP(ColorViewer.ToSharedRef(), &FColorViewer::SetToneMapping, FToneMapping::EMode::RGB),
+			FCanExecuteAction(),
+			FIsActionChecked::CreateLambda([this]
+			{
+				return ColorViewer->GetToneMapping() == FToneMapping::EMode::RGB;
+			})
+		);
+
+		CommandList->MapAction(
+			Commands.ToneMappingLum,
+			FExecuteAction::CreateSP(ColorViewer.ToSharedRef(), &FColorViewer::SetToneMapping, FToneMapping::EMode::Lum),
+			FCanExecuteAction(),
+			FIsActionChecked::CreateLambda([this]
+			{
+				return ColorViewer->GetToneMapping() == FToneMapping::EMode::Lum;
+			})
+		);
+	}
+}
+
+#undef LOCTEXT_NAMESPACE
+
+#endif // IMAGE_WIDGETS_BUILD_COLOR_VIEWER_SAMPLE
diff --git a/Engine/Plugins/Experimental/ImageWidgets/Source/ImageWidgets/Private/ColorViewerSample/ColorViewerWidget.h b/Engine/Plugins/Experimental/ImageWidgets/Source/ImageWidgets/Private/ColorViewerSample/ColorViewerWidget.h
new file mode 100644
--- /dev/null
+++ b/Engine/Plugins/Experimental/ImageWidgets/Source/ImageWidgets/Private/ColorViewerSample/ColorViewerWidget.h
@@ -0,0 +1,50 @@
+﻿// Copyright Epic Games, Inc. All Rights Reserved.
+
+#pragma once
+
+#if IMAGE_WIDGETS_BUILD_COLOR_VIEWER_SAMPLE
+
+#include "ColorViewer.h"
+
+#include <SImageViewport.h>
+
+namespace UE::ImageWidgets::Sample
+{
+	/**
+	 * Widget that contains and configures the image widgets.
+	 */
+	class SColorViewerWidget : public SCompoundWidget
+	{
+	public:
+		SLATE_BEGIN_ARGS(SColorViewerWidget)
+			{
+			}
+
+		SLATE_END_ARGS()
+
+		/** Gets called by Slate for construction of this widget. */
+		void Construct(const FArguments& Args);
+
+		/** Forwards key presses to the command bindings. */
+		virtual FReply OnKeyDown(const FGeometry& Geometry, const FKeyEvent& KeyEvent) override;
+
+	private:
+		/** Functions for adding tone mapping buttons as viewport toolbar extensions. */
+		void AddRandomColorButton(FToolBarBuilder& ToolbarBuilder) const;
+		void AddToneMappingButtons(FToolBarBuilder& ToolbarBuilder) const;
+
+		/** Binds all of the commands. */
+		void BindCommands();
+
+		/** The image viewer implementation that contains the image data and renders the image. */
+		TSharedPtr<FColorViewer> ColorViewer;
+
+		/** The image viewport in which the image gets displayed in. */
+		TSharedPtr<SImageViewport> Viewport;
+
+		/** The commands used by this sample widget. */
+		TSharedPtr<FUICommandList> CommandList;
+	};
+}
+
+#endif // IMAGE_WIDGETS_BUILD_COLOR_VIEWER_SAMPLE
diff --git a/Engine/Plugins/Experimental/ImageWidgets/Source/ImageWidgets/Private/ImageViewportClient.cpp b/Engine/Plugins/Experimental/ImageWidgets/Source/ImageWidgets/Private/ImageViewportClient.cpp
new file mode 100644
--- /dev/null
+++ b/Engine/Plugins/Experimental/ImageWidgets/Source/ImageWidgets/Private/ImageViewportClient.cpp
@@ -0,0 +1,376 @@
+﻿// Copyright Epic Games, Inc. All Rights Reserved.
+
+#include "ImageViewportClient.h"
+
+#include "CanvasItem.h"
+#include "CanvasTypes.h"
+#include "ImageWidgetsLogCategory.h"
+#include "SImageViewport.h"
+#include "Texture2DPreview.h"
+#include "TextureResource.h"
+#include "Engine/Texture2D.h"
+
+#define LOCTEXT_NAMESPACE "ImageViewportClient"
+
+namespace UE::ImageWidgets
+{
+	
+UTexture2D* CreateCheckerTexture(const FColor& Color1, const FColor& Color2, const int32 CheckerSize)
+{
+	constexpr EPixelFormat Format = PF_B8G8R8A8;
+	const int32 TextureSize = FMath::Min<int32>(FMath::RoundUpToPowerOfTwo(FMath::Max(GPixelFormats[PF_B8G8R8A8].BlockSizeX, CheckerSize)) * 2, 4096);
+
+	UTexture2D* CheckerTexture = UTexture2D::CreateTransient(TextureSize, TextureSize, Format);
+	check(CheckerTexture != nullptr);
+
+	UE_LOG(LogImageWidgets, Verbose, TEXT("Creating background checker texture '%s' with Color1=%s, Color2=%s, and Size=%d."),
+	       *CheckerTexture->GetName(), *Color1.ToHex(), *Color2.ToHex(), CheckerSize);
+	
+	// TODO Revisit setting the LODGroup, and figure out if there is a better way to do this.
+	// Setting the LOD group seems to be the only way to avoid interpolation when zoomed in. 
+	CheckerTexture->LODGroup = TEXTUREGROUP_Pixels2D;
+	
+	TIndirectArray<FTexture2DMipMap>& Mips = CheckerTexture->GetPlatformData()->Mips;
+	int32 MipSize = TextureSize;
+	while (MipSize >= GPixelFormats[Format].BlockSizeX)
+	{
+		const int32 MipIndex = MipSize == TextureSize ? 0 : Mips.Add(new FTexture2DMipMap(MipSize, MipSize, 1));
+
+		FByteBulkData& MipBulkData = Mips[MipIndex].BulkData;
+
+		FColor *const MipColorData = [&MipBulkData, MipIndex, MipSize]
+		{
+			FColor* ColorData = static_cast<FColor*>(MipBulkData.Lock(LOCK_READ_WRITE));
+			if (MipIndex > 0)
+			{
+				const int64 NumBlocks = MipSize / GPixelFormats[Format].BlockSizeX;
+				const int64 NumBytes = NumBlocks * NumBlocks * GPixelFormats[Format].BlockBytes;
+				
+				ColorData = reinterpret_cast<FColor*>(MipBulkData.Realloc(NumBytes));
+			}
+			return ColorData;
+		}();
+
+		if (MipSize > 1)
+		{
+			const int32 MipCheckerSize = MipSize / 2;
+			for (int32 Y = 0; Y < MipSize; ++Y)
+			{
+				const bool bTop = Y < MipCheckerSize;
+				for (int32 X = 0; X < MipSize; ++X)
+				{
+					const bool bLeft = X < MipCheckerSize;
+					MipColorData[X + Y * MipSize] = bTop == bLeft ? Color1 : Color2;
+				}
+			}
+		}
+		else
+		{
+			*MipColorData = FColor((static_cast<uint16>(Color1.R) + Color2.R) / 2,
+			                       (static_cast<uint16>(Color1.G) + Color2.G) / 2,
+			                       (static_cast<uint16>(Color1.B) + Color2.B) / 2);
+		}
+
+		MipBulkData.Unlock();
+
+		MipSize /= 2;
+	}
+
+	CheckerTexture->UpdateResource();
+
+	return CheckerTexture;
+}
+
+void DestroyCheckerTexture(TStrongObjectPtr<UTexture2D>& CheckerTexture)
+{
+	if (CheckerTexture)
+	{
+		UE_LOG(LogImageWidgets, Verbose, TEXT("Destroying background checker texture '%s'."), *CheckerTexture->GetName());
+
+		if (CheckerTexture->GetResource())
+		{
+			CheckerTexture->ReleaseResource();
+		}
+		CheckerTexture->MarkAsGarbage();
+
+		CheckerTexture = nullptr;
+	}
+}
+
+FImageViewportClient::FImageViewportClient(const TWeakPtr<SEditorViewport>& InEditorViewport, FGetImageSize&& InGetImageSize, FDrawImage&& InDrawImage,
+                                           FGetDrawSettings&& InGetDrawSettings)
+	: FEditorViewportClient(nullptr, nullptr, InEditorViewport)
+	, GetImageSize(MoveTemp(InGetImageSize))
+	, DrawImage(MoveTemp(InDrawImage))
+	, GetDrawSettings(MoveTemp(InGetDrawSettings))
+{
+	check(GetImageSize.IsBound());
+	check(DrawImage.IsBound());
+	check(GetDrawSettings.IsBound());
+
+	SetRealtime(true);
+}
+
+FImageViewportClient::~FImageViewportClient()
+{
+	if (CheckerTexture)
+	{
+		DestroyCheckerTexture(CheckerTexture);
+	}
+}
+
+void FImageViewportClient::Draw(FViewport* InViewport, FCanvas* Canvas)
+{
+	const SImageViewport::FDrawSettings DrawSettings = GetDrawSettings.Execute();
+	CreateOrDestroyCheckerTextureIfSettingsChanged(DrawSettings);
+
+	Canvas->Clear(DrawSettings.ClearColor);
+
+	const FIntPoint ImageSize = GetImageSize.Execute();
+	if (ImageSize.X <= 0 || ImageSize.Y <= 0)
+	{
+		return;
+	}
+
+	CachedPlacement = GetPlacementProperties({InViewport->GetSizeXY()}, {ImageSize});
+	bCachedPlacementIsValid = true;
+
+	const IImageViewer::FDrawProperties::FMip MipProperties = GetMipProperties();
+
+	// Draw border
+	if (DrawSettings.bBorderEnabled)
+	{
+		FCanvasBoxItem Border(CachedPlacement.Offset - DrawSettings.BorderThickness / 2.0, CachedPlacement.Size + DrawSettings.BorderThickness);
+		Border.LineThickness = DrawSettings.BorderThickness;
+		Border.SetColor(DrawSettings.BorderColor);
+		Canvas->DrawItem(Border);
+	}
+
+	// Draw background
+	if (DrawSettings.bBackgroundColorEnabled || DrawSettings.bBackgroundCheckerEnabled)
+	{
+		FCanvasTileItem Background(CachedPlacement.Offset, CachedPlacement.Size, DrawSettings.BackgroundColor);
+
+		TRefCountPtr<FBatchedElementTexture2DPreviewParameters> BatchedElementParameters;
+
+		if (DrawSettings.bBackgroundCheckerEnabled)
+		{
+			if (ensure(CheckerTexture != nullptr))
+			{
+				Background.SetColor(FLinearColor::White);
+				Background.Texture = CheckerTexture->GetResource();
+
+				Background.UV1 = FVector2D(ImageSize) / CheckerTexture->GetSizeX();
+
+				BatchedElementParameters = new FBatchedElementTexture2DPreviewParameters(
+					-1, 0, 0, false, false, false, false, false, false/*MipProperties.bUsePointSampling*/);
+				Background.BatchedElementParameters = BatchedElementParameters.GetReference();
+			}
+		}
+
+		Canvas->DrawItem(Background);
+	}
+
+	// Draw image
+	DrawImage.Execute(InViewport, Canvas, {CachedPlacement, MipProperties});
+}
+
+EMouseCursor::Type FImageViewportClient::GetCursor(FViewport* InViewport, int32 X, int32 Y)	
+{
+	if (bDragging)
+	{
+		CachedMouseX = X;
+		CachedMouseY = Y;
+		return EMouseCursor::GrabHandClosed;
+	}
+
+	const auto [PixelCoordsValid, PixelCoords] = GetPixelCoordinatesUnderCursor();
+	const FIntPoint ImageSize = GetImageSize.Execute();
+	if (PixelCoordsValid && 0 <= PixelCoords.X && 0 <= PixelCoords.Y && PixelCoords.X < ImageSize.X && PixelCoords.Y < ImageSize.Y)
+	{
+		return EMouseCursor::Crosshairs;
+	}
+
+	return EMouseCursor::Default;
+}
+
+bool FImageViewportClient::InputKey(const FInputKeyEventArgs& EventArgs)
+{
+	if (EventArgs.Event == IE_Pressed)
+	{
+		if (EventArgs.Key == EKeys::MouseScrollUp)
+		{
+			const auto [PixelCoordsValid, PixelCoords] = GetPixelCoordinatesUnderCursor();
+			if (PixelCoordsValid)
+			{
+				Controller.ZoomIn(PixelCoords, GetImageSize.Execute());
+			}
+			return true;
+		}
+
+		if (EventArgs.Key == EKeys::MouseScrollDown)
+		{
+			const auto [PixelCoordsValid, PixelCoords] = GetPixelCoordinatesUnderCursor();
+			if (PixelCoordsValid)
+			{
+				Controller.ZoomOut(PixelCoords, GetImageSize.Execute());
+			}
+			return true;
+		}
+
+		if (EventArgs.Key == EKeys::F)
+		{
+			const FIntPoint ImageSize = GetImageSize.Execute();
+			const FIntPoint ViewportSize = Viewport->GetSizeXY();
+			Controller.Reset(ImageSize, ViewportSize);
+
+			return true;
+		}
+	}
+
+	return FEditorViewportClient::InputKey(EventArgs);
+}
+
+void FImageViewportClient::TrackingStarted(const FInputEventState& InputState, bool bIsDraggingWidget, bool bNudge)
+{
+	if (!bIsDraggingWidget)
+	{
+		FIntPoint MousePos;
+		InputState.GetViewport()->GetMousePos(MousePos);
+
+		if (InputState.IsRightMouseButtonPressed())
+		{
+			bDragging = true;
+			DraggingStart = MousePos;
+			return;
+		}
+	}
+
+	FEditorViewportClient::TrackingStarted(InputState, bIsDraggingWidget, bNudge);
+}
+
+void FImageViewportClient::TrackingStopped()
+{
+	if (bDragging)
+	{
+		bDragging = false;
+
+		FIntPoint DraggingEnd;
+		Viewport->GetMousePos(DraggingEnd);
+		Controller.Pan(DraggingEnd - DraggingStart);
+	}
+
+	RequiredCursorVisibiltyAndAppearance.bDontResetCursor = true;
+
+	FEditorViewportClient::TrackingStopped();
+}
+
+int32 FImageViewportClient::GetMipLevel() const
+{
+	return MipLevel;
+}
+
+void FImageViewportClient::SetMipLevel(int32 InMipLevel)
+{
+	MipLevel = InMipLevel;
+}
+
+void FImageViewportClient::ResetController(FIntPoint ImageSize)
+{
+	Controller.Reset(ImageSize, Viewport->GetSizeXY());
+}
+
+void FImageViewportClient::ResetZoom(FIntPoint ImageSize)
+{		
+	const FImageViewportController::FZoomSettings ZoomSettings = Controller.GetZoom();
+	Controller.SetZoom(ZoomSettings.Mode, ZoomSettings.Zoom, ImageSize, Viewport->GetSizeXY());
+}
+
+FImageViewportController::FZoomSettings FImageViewportClient::GetZoom() const
+{
+	return Controller.GetZoom();
+}
+	
+void FImageViewportClient::SetZoom(FImageViewportController::EZoomMode Mode, double Zoom)
+{
+	const FIntPoint ImageSize = GetImageSize.Execute();
+	const FIntPoint ViewportSize = Viewport->GetSizeXY();
+
+	Controller.SetZoom(Mode, Zoom, ImageSize, ViewportSize);
+}
+
+TPair<bool, FVector2d> FImageViewportClient::GetPixelCoordinatesUnderCursor() const
+{
+	if (!bCachedPlacementIsValid || CurrentMousePos == FIntPoint{-1, -1})
+	{
+		return {false, FVector2d::Zero()};
+	}
+
+	const FVector2d MousePos(CurrentMousePos.X + 0.5, CurrentMousePos.Y + 0.5);
+	const FVector2d CurrentDrag = GetCurrentDrag();
+
+	FVector2d CursorPos = FVector2d(MousePos) - CachedPlacement.Offset + CurrentDrag;
+	CursorPos /= CachedPlacement.ZoomFactor;
+
+	return {true, CursorPos};
+}
+
+FVector2d FImageViewportClient::GetCurrentDrag() const
+{
+	if (bDragging)
+	{
+		FIntPoint DraggingEnd;
+		Viewport->GetMousePos(DraggingEnd);
+		return DraggingEnd - DraggingStart;
+	}
+	return FVector2d::Zero();;
+}
+
+IImageViewer::FDrawProperties::FPlacement FImageViewportClient::GetPlacementProperties(FVector2d ViewportSize, FVector2d ImageSize) const
+{
+	const FVector2d CurrentDrag = GetCurrentDrag();
+	const FVector2d Pan = Controller.GetPan(CurrentDrag);
+
+	const double ZoomFactor = Controller.GetZoom().Zoom;
+
+	const FVector2d TileSize = ImageSize * ZoomFactor;
+	const FVector2d TileOffset = (ViewportSize - TileSize) / 2 + Pan;
+
+	return {TileOffset, TileSize, ZoomFactor};
+}
+
+IImageViewer::FDrawProperties::FMip FImageViewportClient::GetMipProperties() const
+{
+	const double ZoomFactor = Controller.GetZoom().Zoom;
+	const float MipFactor = 1.0f / FMath::Pow(2.0f, MipLevel);
+	const float EffectiveMipLevel = ZoomFactor < MipFactor ? -1.0f : MipLevel;
+
+	return {EffectiveMipLevel};
+}
+
+void FImageViewportClient::CreateOrDestroyCheckerTextureIfSettingsChanged(const SImageViewport::FDrawSettings& DrawSettings)
+{
+	const FCheckerTextureSettings NewCheckerTextureSettings{
+		DrawSettings.bBackgroundCheckerEnabled, DrawSettings.BackgroundCheckerColor1, DrawSettings.BackgroundCheckerColor2, DrawSettings.BackgroundCheckerSize
+	};
+	
+	if (CachedCheckerTextureSettings != NewCheckerTextureSettings)
+	{
+		if (CheckerTexture)
+		{
+		    DestroyCheckerTexture(CheckerTexture);
+        }
+        
+        if (NewCheckerTextureSettings.bEnabled)
+		{
+			CheckerTexture = TStrongObjectPtr<UTexture2D>(CreateCheckerTexture(NewCheckerTextureSettings.Color1.ToFColorSRGB(),
+			                                                                   NewCheckerTextureSettings.Color2.ToFColorSRGB(),
+			                                                                   NewCheckerTextureSettings.CheckerSize));
+		}
+
+		CachedCheckerTextureSettings = NewCheckerTextureSettings;
+	}
+}
+}
+
+#undef LOCTEXT_NAMESPACE
diff --git a/Engine/Plugins/Experimental/ImageWidgets/Source/ImageWidgets/Private/ImageViewportClient.h b/Engine/Plugins/Experimental/ImageWidgets/Source/ImageWidgets/Private/ImageViewportClient.h
new file mode 100644
--- /dev/null
+++ b/Engine/Plugins/Experimental/ImageWidgets/Source/ImageWidgets/Private/ImageViewportClient.h
@@ -0,0 +1,81 @@
+﻿// Copyright Epic Games, Inc. All Rights Reserved.
+
+#pragma once
+
+#include "EditorViewportClient.h"
+#include "ImageViewportController.h"
+#include "IImageViewer.h"
+#include "SImageViewport.h"
+
+namespace UE::ImageWidgets
+{
+	class FABComparison;
+
+	DECLARE_DELEGATE_RetVal(FIntPoint, FGetImageSize)
+	DECLARE_DELEGATE_ThreeParams(FDrawImage, FViewport*, FCanvas*, const IImageViewer::FDrawProperties&)
+	DECLARE_DELEGATE_RetVal(SImageViewport::FDrawSettings, FGetDrawSettings)
+
+	/**
+	 * Viewport client for controlling the camera and drawing viewport contents. 
+	 */
+	class FImageViewportClient : public FEditorViewportClient
+	{
+	public:
+		FImageViewportClient(const TWeakPtr<SEditorViewport>& InViewport, FGetImageSize&& InGetImageSize, FDrawImage&& InDrawImage,
+		                     FGetDrawSettings&& InGetDrawSettings);
+		virtual ~FImageViewportClient() override;
+
+		virtual void Draw(FViewport* InViewport, FCanvas* Canvas) override;
+
+		// SEditorViewport overrides - begin
+		virtual EMouseCursor::Type GetCursor(FViewport* InViewport, int32 X, int32 Y) override;
+		virtual bool InputKey(const FInputKeyEventArgs& EventArgs) override;
+		virtual void TrackingStarted(const FInputEventState& InputState, bool bIsDraggingWidget, bool bNudge) override;
+		virtual void TrackingStopped() override;
+		// SEditorViewport overrides - end
+
+		int32 GetMipLevel() const;
+		void SetMipLevel(int32 InMipLevel);
+		FImageViewportController::FZoomSettings GetZoom() const;
+		void SetZoom(FImageViewportController::EZoomMode Mode, double Zoom);
+		void ResetController(FIntPoint ImageSize);
+		void ResetZoom(FIntPoint ImageSize);
+
+		TPair<bool, FVector2d> GetPixelCoordinatesUnderCursor() const;
+
+	private:
+		struct FCheckerTextureSettings
+		{
+			bool operator==(const FCheckerTextureSettings&) const = default;
+
+			bool bEnabled = false;
+			FLinearColor Color1;
+			FLinearColor Color2;
+			uint32 CheckerSize = 0;
+		};
+
+		FVector2d GetCurrentDrag() const;
+		IImageViewer::FDrawProperties::FPlacement GetPlacementProperties(FVector2d ViewportSize, FVector2d TextureSize) const;
+		IImageViewer::FDrawProperties::FMip GetMipProperties() const;
+
+		void CreateOrDestroyCheckerTextureIfSettingsChanged(const SImageViewport::FDrawSettings& DrawSettings);
+
+		FGetImageSize GetImageSize;
+		FDrawImage DrawImage;
+		FGetDrawSettings GetDrawSettings;
+
+		bool bDragging = false;
+		FIntPoint DraggingStart;
+
+		bool bMipSelected = true;
+		int32 MipLevel = -1;
+
+		IImageViewer::FDrawProperties::FPlacement CachedPlacement;
+		bool bCachedPlacementIsValid = false;
+
+		FImageViewportController Controller;
+
+		TStrongObjectPtr<UTexture2D> CheckerTexture = nullptr;
+		FCheckerTextureSettings CachedCheckerTextureSettings;
+	};
+}
diff --git a/Engine/Plugins/Experimental/ImageWidgets/Source/ImageWidgets/Private/ImageViewportController.cpp b/Engine/Plugins/Experimental/ImageWidgets/Source/ImageWidgets/Private/ImageViewportController.cpp
new file mode 100644
--- /dev/null
+++ b/Engine/Plugins/Experimental/ImageWidgets/Source/ImageWidgets/Private/ImageViewportController.cpp
@@ -0,0 +1,160 @@
+// Copyright Epic Games, Inc. All Rights Reserved.
+
+#include "ImageViewportController.h"
+
+namespace UE::ImageWidgets
+{
+	namespace ImageViewportController_Local
+	{
+		const TArray<double>& GetZoomLevels()
+		{
+			constexpr double ZoomMin = 1.0 / 64;
+			constexpr double ZoomMax = 64.0;
+			constexpr int32 ZoomStepsToDouble = 8;
+			constexpr double ZoomPowerIncrement = 1.0 / ZoomStepsToDouble;
+
+			static TArray<double> ZoomLevels = []
+			{
+				const int32 ZoomMinPower = FMath::RoundToInt32(FMath::Log2(ZoomMin));
+				const int32 ZoomMaxPower = FMath::RoundToInt32(FMath::Log2(ZoomMax));
+				const int32 Num = (ZoomMaxPower - ZoomMinPower) * ZoomStepsToDouble + 1;
+
+				TArray<double> Result;
+				Result.Reserve(Num);
+
+				for (double Power = ZoomMinPower; Power <= ZoomMaxPower; Power += ZoomPowerIncrement)
+				{
+					Result.Add(FMath::Pow(2.0, Power));
+				}
+
+				return Result;
+			}();
+
+			return ZoomLevels;
+		}
+
+		double ZoomIn(const double CurrentZoom)
+		{
+			const TArray<double>& ZoomLevels = GetZoomLevels();
+			const int32 Index = Algo::LowerBound(ZoomLevels, CurrentZoom);
+
+			if (Index >= ZoomLevels.Num())
+			{
+				return ZoomLevels.Last();
+			}
+
+			if (CurrentZoom < ZoomLevels[Index])
+			{
+				return ZoomLevels[Index];
+			}
+
+			return ZoomLevels[FMath::Min(Index + 1, ZoomLevels.Num() - 1)];
+		}
+
+		double ZoomOut(const double CurrentZoom)
+		{
+			const TArray<double>& ZoomLevels = GetZoomLevels();
+			const int32 Index = Algo::LowerBound(ZoomLevels, CurrentZoom);
+
+			if (Index >= ZoomLevels.Num())
+			{
+				return ZoomLevels.Last();
+			}
+
+			if (CurrentZoom > ZoomLevels[Index])
+			{
+				return ZoomLevels[Index];
+			}
+
+			return ZoomLevels[FMath::Max(0, Index - 1)];
+		}
+
+		FVector2d AdjustPan(const FVector2d OldPan, const double OldZoom, const double NewZoom, const FVector2d CursorPos, const FVector2d& ImageSize)
+		{
+			const FVector2d ImageCenter = ImageSize * 0.5;
+			const FVector2d DistanceCursorToCenter = CursorPos - ImageCenter;
+			const FVector2d DeltaDistanceCursorToCenter = DistanceCursorToCenter * (NewZoom - OldZoom);
+			const FVector2d NewPan = OldPan - DeltaDistanceCursorToCenter;
+
+			return NewPan;
+		}
+	}
+
+	FImageViewportController::FImageViewportController()
+	{
+		Reset({0, 0}, {0, 0});
+	}
+
+	void FImageViewportController::Pan(FVector2d Delta)
+	{
+		PanAmount += Delta;
+	}
+
+	void FImageViewportController::Reset(FIntPoint ImageSize, FIntPoint ViewportSize)
+	{
+		SetZoom(EZoomMode::Fit, 1.0, ImageSize, ViewportSize);
+
+		PanAmount = FVector2d::Zero();
+	}
+
+	void FImageViewportController::ZoomIn(FVector2d CursorPos, const FIntPoint& ImageSize)
+	{
+		const double OldZoom = ZoomSettings.Zoom;
+
+		ZoomSettings.Mode = EZoomMode::Custom;
+		ZoomSettings.Zoom = ImageViewportController_Local::ZoomIn(ZoomSettings.Zoom);
+
+		PanAmount = ImageViewportController_Local::AdjustPan(PanAmount, OldZoom, ZoomSettings.Zoom, CursorPos, FVector2d(ImageSize.X, ImageSize.Y));
+	}
+
+	void FImageViewportController::ZoomOut(FVector2d CursorPos, const FIntPoint& ImageSize)
+	{
+		const double OldZoom = ZoomSettings.Zoom;
+
+		ZoomSettings.Mode = EZoomMode::Custom;
+		ZoomSettings.Zoom = ImageViewportController_Local::ZoomOut(ZoomSettings.Zoom);
+
+		PanAmount = ImageViewportController_Local::AdjustPan(PanAmount, OldZoom, ZoomSettings.Zoom, CursorPos, FVector2d(ImageSize.X, ImageSize.Y));
+	}
+
+	FVector2d FImageViewportController::GetPan(FVector2d Drag) const
+	{
+		return PanAmount + Drag;
+	}
+
+	FImageViewportController::FZoomSettings FImageViewportController::GetZoom() const
+	{
+		return ZoomSettings;
+	}
+
+	void FImageViewportController::SetZoom(const EZoomMode ZoomMode, const double Zoom, const FIntPoint& ImageSize, const FIntPoint& ViewportSize)
+	{
+		ZoomSettings.Mode = ZoomMode;
+
+		if (ZoomSettings.Mode == EZoomMode::Custom)
+		{
+			ZoomSettings.Zoom = Zoom;
+			return;
+		}
+
+		if (ImageSize == FIntPoint::ZeroValue || ViewportSize == FIntPoint::ZeroValue)
+		{
+			ZoomSettings.Zoom = 1.0;
+			return;
+		}
+
+		const FVector2d SizeRatio = FVector2d(ViewportSize) / ImageSize;
+		const double SizeRatioMin = SizeRatio.GetMin();
+
+		if (ZoomSettings.Mode == EZoomMode::Fill)
+		{
+			ZoomSettings.Zoom = SizeRatioMin;
+		}
+		else
+		{
+			ZoomSettings.Zoom = FMath::Min(1.0, SizeRatioMin);
+		}
+
+		PanAmount = FVector2d::Zero();
+	}
+}
diff --git a/Engine/Plugins/Experimental/ImageWidgets/Source/ImageWidgets/Private/ImageViewportController.h b/Engine/Plugins/Experimental/ImageWidgets/Source/ImageWidgets/Private/ImageViewportController.h
new file mode 100644
--- /dev/null
+++ b/Engine/Plugins/Experimental/ImageWidgets/Source/ImageWidgets/Private/ImageViewportController.h
@@ -0,0 +1,38 @@
+// Copyright Epic Games, Inc. All Rights Reserved.
+
+#pragma once
+
+#include "CoreMinimal.h"
+
+namespace UE::ImageWidgets
+{
+	/**
+	 * Camera controller for the 2D viewport supporting panning and zooming.
+	 */
+	class FImageViewportController
+	{
+	public:
+		enum class EZoomMode { Custom, Fit, Fill };
+
+		struct FZoomSettings
+		{
+			EZoomMode Mode;
+			double Zoom;
+		};
+
+		FImageViewportController();
+
+		void Pan(FVector2d Delta);
+		void Reset(FIntPoint ImageSize, FIntPoint ViewportSize);
+		void ZoomIn(FVector2d CursorPos, const FIntPoint& ImageSize);
+		void ZoomOut(FVector2d CursorPos, const FIntPoint& ImageSize);
+
+		FVector2d GetPan(FVector2d Drag) const;
+		FZoomSettings GetZoom() const;
+		void SetZoom(EZoomMode ZoomMode, double Zoom, const FIntPoint& ImageSize, const FIntPoint& ViewportSize);
+
+	private:
+		FVector2d PanAmount;
+		FZoomSettings ZoomSettings;
+	};
+}
diff --git a/Engine/Plugins/Experimental/ImageWidgets/Source/ImageWidgets/Private/ImageWidgetsCommands.cpp b/Engine/Plugins/Experimental/ImageWidgets/Source/ImageWidgets/Private/ImageWidgetsCommands.cpp
new file mode 100644
--- /dev/null
+++ b/Engine/Plugins/Experimental/ImageWidgets/Source/ImageWidgets/Private/ImageWidgetsCommands.cpp
@@ -0,0 +1,37 @@
+﻿// Copyright Epic Games, Inc. All Rights Reserved.
+
+#include "ImageWidgetsCommands.h"
+
+#include "Framework/Commands/InputChord.h"
+#include "Styling/AppStyle.h"
+
+#define LOCTEXT_NAMESPACE "ImageWidgetsCommands"
+
+namespace UE::ImageWidgets
+{
+	FImageWidgetsCommands::FImageWidgetsCommands()
+		: TCommands("ImageWidgets", LOCTEXT("ContextDescription", "Image Widgets"), NAME_None, FAppStyle::Get().GetStyleSetName())
+	{
+	}
+
+	void FImageWidgetsCommands::RegisterCommands()
+	{
+		UI_COMMAND(MipMinus, "Lower", "Show next lower mip", EUserInterfaceActionType::None, FInputChord(EKeys::Subtract));
+		UI_COMMAND(MipPlus, "Higher", "Show next higher mip", EUserInterfaceActionType::None, FInputChord(EKeys::Add));
+
+		UI_COMMAND(ToggleOverlay, "Toggle Overlay", "Toggles the visibility of the toolbar and the status bar.", EUserInterfaceActionType::ToggleButton,
+		           FInputChord(EKeys::O));
+
+		UI_COMMAND(Zoom12, "12.5%", "Set Zoom to 12.5%", EUserInterfaceActionType::Button, FInputChord());
+		UI_COMMAND(Zoom25, "25%", "Set Zoom to 25%", EUserInterfaceActionType::Button, FInputChord());
+		UI_COMMAND(Zoom50, "50%", "Set Zoom to 50%", EUserInterfaceActionType::Button, FInputChord());
+		UI_COMMAND(Zoom100, "100%", "Set Zoom to 100%", EUserInterfaceActionType::Button, FInputChord());
+		UI_COMMAND(Zoom200, "200%", "Set Zoom to 200%", EUserInterfaceActionType::Button, FInputChord());
+		UI_COMMAND(Zoom400, "400%", "Set Zoom to 400%", EUserInterfaceActionType::Button, FInputChord());
+		UI_COMMAND(Zoom800, "800%", "Set Zoom to 800%", EUserInterfaceActionType::Button, FInputChord());
+		UI_COMMAND(ZoomFit, "Scale to Fit", "Optionally downscales the image to fit within the viewport", EUserInterfaceActionType::RadioButton, FInputChord());
+		UI_COMMAND(ZoomFill, "Scale to Fill", "Downscales or upscales the image to fill the viewport", EUserInterfaceActionType::RadioButton, FInputChord());
+	}
+}
+
+#undef LOCTEXT_NAMESPACE
diff --git a/Engine/Plugins/Experimental/ImageWidgets/Source/ImageWidgets/Private/ImageWidgetsCommands.h b/Engine/Plugins/Experimental/ImageWidgets/Source/ImageWidgets/Private/ImageWidgetsCommands.h
new file mode 100644
--- /dev/null
+++ b/Engine/Plugins/Experimental/ImageWidgets/Source/ImageWidgets/Private/ImageWidgetsCommands.h
@@ -0,0 +1,35 @@
+﻿// Copyright Epic Games, Inc. All Rights Reserved.
+
+#pragma once
+
+#include "Framework/Commands/Commands.h"
+#include "Framework/Commands/UICommandInfo.h"
+
+namespace UE::ImageWidgets
+{
+	/**
+	 * Provides commands used for the image widgets.
+	 */
+	class FImageWidgetsCommands : public TCommands<FImageWidgetsCommands>
+	{
+	public:
+		FImageWidgetsCommands();
+
+		virtual void RegisterCommands() override;
+
+		TSharedPtr<FUICommandInfo> MipMinus;
+		TSharedPtr<FUICommandInfo> MipPlus;
+
+		TSharedPtr<FUICommandInfo> ToggleOverlay;
+
+		TSharedPtr<FUICommandInfo> Zoom12;
+		TSharedPtr<FUICommandInfo> Zoom25;
+		TSharedPtr<FUICommandInfo> Zoom50;
+		TSharedPtr<FUICommandInfo> Zoom100;
+		TSharedPtr<FUICommandInfo> Zoom200;
+		TSharedPtr<FUICommandInfo> Zoom400;
+		TSharedPtr<FUICommandInfo> Zoom800;
+		TSharedPtr<FUICommandInfo> ZoomFit;
+		TSharedPtr<FUICommandInfo> ZoomFill;
+	};
+}
diff --git a/Engine/Plugins/Experimental/ImageWidgets/Source/ImageWidgets/Private/ImageWidgetsLogCategory.cpp b/Engine/Plugins/Experimental/ImageWidgets/Source/ImageWidgets/Private/ImageWidgetsLogCategory.cpp
new file mode 100644
--- /dev/null
+++ b/Engine/Plugins/Experimental/ImageWidgets/Source/ImageWidgets/Private/ImageWidgetsLogCategory.cpp
@@ -0,0 +1,8 @@
+﻿// Copyright Epic Games, Inc. All Rights Reserved.
+
+#include "ImageWidgetsLogCategory.h"
+
+namespace UE::ImageWidgets
+{
+	DEFINE_LOG_CATEGORY(LogImageWidgets);
+}
diff --git a/Engine/Plugins/Experimental/ImageWidgets/Source/ImageWidgets/Private/ImageWidgetsLogCategory.h b/Engine/Plugins/Experimental/ImageWidgets/Source/ImageWidgets/Private/ImageWidgetsLogCategory.h
new file mode 100644
--- /dev/null
+++ b/Engine/Plugins/Experimental/ImageWidgets/Source/ImageWidgets/Private/ImageWidgetsLogCategory.h
@@ -0,0 +1,10 @@
+﻿// Copyright Epic Games, Inc. All Rights Reserved.
+
+#pragma once
+
+#include <Logging/LogMacros.h>
+
+namespace UE::ImageWidgets
+{
+	DECLARE_LOG_CATEGORY_EXTERN(LogImageWidgets, Log, All)
+}
diff --git a/Engine/Plugins/Experimental/ImageWidgets/Source/ImageWidgets/Private/ImageWidgetsModule.cpp b/Engine/Plugins/Experimental/ImageWidgets/Source/ImageWidgets/Private/ImageWidgetsModule.cpp
new file mode 100644
--- /dev/null
+++ b/Engine/Plugins/Experimental/ImageWidgets/Source/ImageWidgets/Private/ImageWidgetsModule.cpp
@@ -0,0 +1,116 @@
+﻿// Copyright Epic Games, Inc. All Rights Reserved.
+
+#if !IMAGE_WIDGETS_BUILD_COLOR_VIEWER_SAMPLE
+
+#include "Modules/ModuleManager.h"
+
+IMPLEMENT_MODULE(FDefaultModuleImpl, ImageWidgets);
+
+#else // !IMAGE_WIDGETS_BUILD_COLOR_VIEWER_SAMPLE
+
+#include "CoreMinimal.h"
+#include "ColorViewerSample/ColorViewerCommands.h"
+#include "ColorViewerSample/ColorViewerWidget.h"
+#include "WorkspaceMenuStructure.h"
+#include "WorkspaceMenuStructureModule.h"
+#include "Framework/Docking/TabManager.h"
+#include "Modules/ModuleInterface.h"
+#include "Widgets/Docking/SDockTab.h"
+
+#define LOCTEXT_NAMESPACE "ImageWidgetsModule"
+
+namespace UE::ImageWidgets
+{
+	/**
+	 * This module provides the color viewer sample within a tab widget accessible via the "Tools > Miscellaneous" menu.
+	 */
+	class FImageWidgetsModule : public IModuleInterface
+	{
+	public:
+		void RegisterTabSpawners(const TSharedPtr<FWorkspaceItem>& WorkspaceGroup);
+		void UnregisterTabSpawners();
+
+		// IModuleInterface overrides - begin
+		virtual void StartupModule() override;
+		virtual void ShutdownModule() override;
+		// IModuleInterface overrides - end
+
+	private:
+		TSharedRef<SDockTab> MakeColorViewerTab(const FSpawnTabArgs&);
+		TSharedRef<SWidget> GetColorViewer(const TSharedRef<SDockTab>& InParentTab);
+
+		bool bHasRegisteredTabSpawners = false;
+		TWeakPtr<Sample::SColorViewerWidget> ColorViewerPtr;
+	};
+
+	void FImageWidgetsModule::RegisterTabSpawners(const TSharedPtr<FWorkspaceItem>& WorkspaceGroup)
+	{
+		if (bHasRegisteredTabSpawners)
+		{
+			UnregisterTabSpawners();
+		}
+
+		bHasRegisteredTabSpawners = true;
+
+		FTabSpawnerEntry& ColorViewerSpawnerEntry = FGlobalTabmanager::Get()->RegisterNomadTabSpawner("ColorViewer",
+			FOnSpawnTab::CreateRaw(this, &FImageWidgetsModule::MakeColorViewerTab))
+		.SetDisplayName(LOCTEXT("ColorViewerTitle", "Color Viewer Sample"))
+		.SetTooltipText(LOCTEXT("ColorViewerTooltipText", "Open the Color Viewer tab, a sample application for the ImageWidgets plugin."))
+		.SetGroup(WorkspaceMenu::GetMenuStructure().GetDeveloperToolsMiscCategory())
+		.SetIcon(FSlateIcon(FAppStyle::GetAppStyleSetName(), "ClassIcon.Texture2D"));
+
+		if (WorkspaceGroup.IsValid())
+		{
+			ColorViewerSpawnerEntry.SetGroup(WorkspaceGroup.ToSharedRef());
+		}
+	}
+
+	void FImageWidgetsModule::UnregisterTabSpawners()
+	{
+		bHasRegisteredTabSpawners = false;
+
+		FGlobalTabmanager::Get()->UnregisterNomadTabSpawner("ColorViewer");
+	}
+
+	void FImageWidgetsModule::StartupModule()
+	{
+		bHasRegisteredTabSpawners = false;
+		RegisterTabSpawners(nullptr);
+
+		Sample::FColorViewerCommands::Register();
+	}
+
+	void FImageWidgetsModule::ShutdownModule()
+	{
+		Sample::FColorViewerCommands::Unregister();
+
+		UnregisterTabSpawners();
+	}
+
+	TSharedRef<SDockTab> FImageWidgetsModule::MakeColorViewerTab(const FSpawnTabArgs&)
+	{
+		TSharedRef<SDockTab> WidgetReflectorTab = SNew(SDockTab)
+			.TabRole(NomadTab);
+		WidgetReflectorTab->SetContent(GetColorViewer(WidgetReflectorTab));
+		return WidgetReflectorTab;
+	}
+
+	TSharedRef<SWidget> FImageWidgetsModule::GetColorViewer(const TSharedRef<SDockTab>& InParentTab)
+	{
+		TSharedPtr<Sample::SColorViewerWidget> ColorViewer = ColorViewerPtr.Pin();
+
+		if (!ColorViewer.IsValid())
+		{
+			ColorViewer = SNew(Sample::SColorViewerWidget);
+			ColorViewerPtr = ColorViewer;
+		}
+
+		return ColorViewer.ToSharedRef();
+	}
+}
+
+#undef LOCTEXT_NAMESPACE
+
+IMPLEMENT_MODULE(UE::ImageWidgets::FImageWidgetsModule, ImageWidgets)
+
+#endif // !IMAGE_WIDGETS_BUILD_COLOR_VIEWER_SAMPLE
diff --git a/Engine/Plugins/Experimental/ImageWidgets/Source/ImageWidgets/Private/ImageWidgetsStyle.cpp b/Engine/Plugins/Experimental/ImageWidgets/Source/ImageWidgets/Private/ImageWidgetsStyle.cpp
new file mode 100644
--- /dev/null
+++ b/Engine/Plugins/Experimental/ImageWidgets/Source/ImageWidgets/Private/ImageWidgetsStyle.cpp
@@ -0,0 +1,43 @@
+﻿// Copyright Epic Games, Inc. All Rights Reserved.
+
+#include "ImageWidgetsStyle.h"
+
+#include "Styling/AppStyle.h"
+#include "Styling/CoreStyle.h"
+#include "Styling/SlateStyleRegistry.h"
+#include "Styling/SlateTypes.h"
+
+namespace UE::ImageWidgets
+{
+	FName FImageWidgetsStyle::StyleName("ImageViewportStyle");
+
+	FImageWidgetsStyle::FImageWidgetsStyle()
+		: FSlateStyleSet(StyleName)
+	{
+		FSlateStyleSet::SetContentRoot(FPaths::EngineContentDir() / TEXT("Editor/Slate"));
+		FSlateStyleSet::SetCoreContentRoot(FPaths::EngineContentDir() / TEXT("Slate"));
+
+		const FTextBlockStyle NormalText = FAppStyle::Get().GetWidgetStyle<FTextBlockStyle>("NormalText");
+		Set("RichTextBlock.Red", FTextBlockStyle(NormalText).SetColorAndOpacity(FLinearColor(1.0f, 0.1f, 0.1f)));
+		Set("RichTextBlock.Green", FTextBlockStyle(NormalText).SetColorAndOpacity(FLinearColor(0.1f, 1.0f, 0.1f)));
+		Set("RichTextBlock.Blue", FTextBlockStyle(NormalText).SetColorAndOpacity(FLinearColor(0.1f, 0.1f, 1.0f)));
+
+
+		FSlateBrush* BrushTableRowOdd = new FSlateBrush();
+		BrushTableRowOdd->TintColor = FLinearColor(0.5f, 0.5f, 0.5f, 1.0f);
+		Set("TableRowOdd", BrushTableRowOdd);
+
+		FSlateStyleRegistry::RegisterSlateStyle(*this);
+	}
+
+	FImageWidgetsStyle::~FImageWidgetsStyle()
+	{
+		FSlateStyleRegistry::UnRegisterSlateStyle(*this);
+	}
+
+	FImageWidgetsStyle& FImageWidgetsStyle::Get()
+	{
+		static FImageWidgetsStyle Instance;
+		return Instance;
+	}
+}
diff --git a/Engine/Plugins/Experimental/ImageWidgets/Source/ImageWidgets/Private/ImageWidgetsStyle.h b/Engine/Plugins/Experimental/ImageWidgets/Source/ImageWidgets/Private/ImageWidgetsStyle.h
new file mode 100644
--- /dev/null
+++ b/Engine/Plugins/Experimental/ImageWidgets/Source/ImageWidgets/Private/ImageWidgetsStyle.h
@@ -0,0 +1,23 @@
+﻿// Copyright Epic Games, Inc. All Rights Reserved.
+
+#pragma once
+
+#include "Styling/SlateStyle.h"
+
+namespace UE::ImageWidgets
+{
+	/**
+	 * Style declarations for the image widgets.
+	 */
+	class FImageWidgetsStyle : public FSlateStyleSet
+	{
+	public:
+		static FName StyleName;
+
+		static FImageWidgetsStyle& Get();
+
+	private:
+		FImageWidgetsStyle();
+		virtual ~FImageWidgetsStyle() override;
+	};
+}
diff --git a/Engine/Plugins/Experimental/ImageWidgets/Source/ImageWidgets/Private/SImageViewport.cpp b/Engine/Plugins/Experimental/ImageWidgets/Source/ImageWidgets/Private/SImageViewport.cpp
new file mode 100644
--- /dev/null
+++ b/Engine/Plugins/Experimental/ImageWidgets/Source/ImageWidgets/Private/SImageViewport.cpp
@@ -0,0 +1,430 @@
+﻿// Copyright Epic Games, Inc. All Rights Reserved.
+
+#include "SImageViewport.h"
+
+#include "CanvasTypes.h"
+#include "ImageViewportClient.h"
+#include "ImageViewportController.h"
+#include "ImageWidgetsCommands.h"
+#include "ImageWidgetsStyle.h"
+#include "SImageViewportToolbar.h"
+#include "Widgets/Text/SRichTextBlock.h"
+#include "Widgets/Text/STextBlock.h"
+
+#define LOCTEXT_NAMESPACE "ImageViewport"
+
+namespace UE::ImageWidgets
+{
+	class FStatusBarExtension
+	{
+	public:
+		virtual ~FStatusBarExtension() = default;
+
+		FName Hook;
+		EExtensionHook::Position HookPosition = EExtensionHook::Before;
+		TSharedPtr<FUICommandList> CommandList;
+		SImageViewport::FStatusBarExtender::FDelegate Delegate;
+	};
+
+	void SImageViewport::FStatusBarExtender::AddExtension(FName ExtensionHook, EExtensionHook::Position HookPosition,
+	                                                      const TSharedPtr<FUICommandList>& Commands, const FDelegate& Delegate)
+	{
+		const TPimplPtr<FStatusBarExtension>& Extension = Extensions.Emplace_GetRef(MakePimpl<FStatusBarExtension>());
+		Extension->Hook = ExtensionHook;
+		Extension->HookPosition = HookPosition;
+		Extension->CommandList = Commands;
+		Extension->Delegate = Delegate;
+	}
+
+	void SImageViewport::FStatusBarExtender::Apply(FName ExtensionHook, EExtensionHook::Position HookPosition, SHorizontalBox& HorizontalBox) const
+	{
+		for (const TPimplPtr<FStatusBarExtension>& Extension : Extensions)
+		{
+			if (Extension->Hook == ExtensionHook && Extension->HookPosition == HookPosition)
+			{
+				Extension->Delegate.ExecuteIfBound(HorizontalBox);
+			}
+		}
+	}
+
+	SImageViewport::SImageViewport()
+	{
+	}
+
+	SImageViewport::~SImageViewport()
+	{
+		FImageWidgetsCommands::Unregister();
+	}
+
+	void SImageViewport::BindCommands()
+	{
+		const FImageWidgetsCommands& Commands = FImageWidgetsCommands::Get();
+
+		auto ToggleOverlay = [this]
+		{
+			ViewportOverlay->SetVisibility(ViewportOverlay->GetVisibility() == EVisibility::Visible ? EVisibility::Hidden : EVisibility::Visible);
+		};
+
+		CommandList->MapAction(
+			Commands.ToggleOverlay,
+			FExecuteAction::CreateLambda(ToggleOverlay)
+		);
+
+		// ZOOM
+
+		auto MapZoom = [this](const TSharedPtr<FUICommandInfo>& Command, double ZoomFactor)
+		{
+			CommandList->MapAction(
+				Command, FExecuteAction::CreateSP(ImageViewportClient.ToSharedRef(), &FImageViewportClient::SetZoom,
+				                                  FImageViewportController::EZoomMode::Custom, ZoomFactor));
+		};
+
+		MapZoom(Commands.Zoom12, 0.125);
+		MapZoom(Commands.Zoom25, 0.25);
+		MapZoom(Commands.Zoom50, 0.5);
+		MapZoom(Commands.Zoom100, 1.0);
+		MapZoom(Commands.Zoom200, 2.0);
+		MapZoom(Commands.Zoom400, 4.0);
+		MapZoom(Commands.Zoom800, 8.0);
+
+		auto MapZoomFitFill = [this](const TSharedPtr<FUICommandInfo>& Command, FImageViewportController::EZoomMode ZoomMode)
+		{
+			CommandList->MapAction(
+				Command,
+				FExecuteAction::CreateSP(ImageViewportClient.ToSharedRef(), &FImageViewportClient::SetZoom, ZoomMode, 0.0),
+				FCanExecuteAction(),
+				FIsActionChecked::CreateLambda([this, ZoomMode] { return ImageViewportClient->GetZoom().Mode == ZoomMode; }));
+		};
+
+		MapZoomFitFill(Commands.ZoomFit, FImageViewportController::EZoomMode::Fit);
+		MapZoomFitFill(Commands.ZoomFill, FImageViewportController::EZoomMode::Fill);
+
+		// MIPS
+
+		CommandList->MapAction(
+			Commands.MipMinus,
+			FExecuteAction::CreateLambda([this]()
+			{
+				const TSharedPtr<FImageViewportClient> ViewportClient = StaticCastSharedPtr<FImageViewportClient>(Client);
+				const int32 CurrentMip = ViewportClient->GetMipLevel();
+
+				int32 NumMips = 0;
+				if (ImageViewer.IsValid())
+				{
+					const IImageViewer::FImageInfo ImageInfo = ImageViewer->GetCurrentImageInfo();
+					if (ImageInfo.bIsValid)
+					{
+						NumMips = ImageInfo.NumMips;
+					}
+				}
+
+				if (CurrentMip <= NumMips)
+				{
+					ViewportClient->SetMipLevel(CurrentMip + 1);
+				}
+			}));
+		CommandList->MapAction(
+			Commands.MipPlus,
+			FExecuteAction::CreateLambda([this]()
+			{
+				const TSharedPtr<FImageViewportClient> ViewportClient = StaticCastSharedPtr<FImageViewportClient>(Client);
+				const int32 CurrentMip = ViewportClient->GetMipLevel();
+				if (CurrentMip > -1)
+				{
+					ViewportClient->SetMipLevel(CurrentMip - 1);
+				}
+			}));
+	}
+
+	void SImageViewport::Construct(const FArguments& InArgs, const TSharedRef<IImageViewer>& InImageViewer)
+	{
+		FImageWidgetsCommands::Register();
+
+		ImageViewer = InImageViewer;
+		check(ImageViewer.IsValid());
+
+		ToolbarExtender = InArgs._ToolbarExtender;
+		StatusBarExtender = InArgs._StatusBarExtender;
+
+		DrawSettings = InArgs._DrawSettings;
+
+		auto GetImageSize = [this]
+		{
+			if (ImageViewer.IsValid())
+			{
+				const IImageViewer::FImageInfo ImageInfo = ImageViewer->GetCurrentImageInfo();
+				return ImageInfo.Size;
+			}
+			return FIntPoint::ZeroValue;
+		};
+
+		auto DrawImage = [this](FViewport* Viewport, FCanvas* Canvas, const IImageViewer::FDrawProperties& DrawProperties)
+		{
+			ImageViewer->DrawCurrentImage(Viewport, Canvas, DrawProperties);
+		};
+
+		auto GetDrawSettings = [this]
+		{
+			return DrawSettings.Get();
+		};
+
+		ImageViewportClient = MakeShareable(new FImageViewportClient(
+			StaticCastWeakPtr<SEditorViewport>(AsWeak()),
+			FGetImageSize::CreateLambda(GetImageSize),
+			FDrawImage::CreateLambda(DrawImage),
+			FGetDrawSettings::CreateLambda(GetDrawSettings)));
+
+		SEditorViewport::Construct(SEditorViewport::FArguments());
+	}
+
+	void SImageViewport::ResetController(FIntPoint ImageSize) const
+	{
+		ImageViewportClient->ResetController(ImageSize);
+	}
+
+	void SImageViewport::ResetMip() const
+	{
+		const int32 SelectedMip = ImageViewportClient->GetMipLevel();
+		if (SelectedMip != -1)
+		{
+			if (ImageViewer.IsValid())
+			{
+				const IImageViewer::FImageInfo ImageInfo = ImageViewer->GetCurrentImageInfo();
+				if (ImageInfo.bIsValid)
+				{
+					if (ImageInfo.NumMips <= SelectedMip)
+					{
+						ImageViewportClient->SetMipLevel(ImageInfo.NumMips - 1);
+					}
+				}
+			}
+		}
+	}
+
+	void SImageViewport::ResetZoom(FIntPoint ImageSize) const
+	{
+		ImageViewportClient->ResetZoom(ImageSize);
+	}
+
+	TSharedPtr<SViewportToolBar> SImageViewport::GetParentToolbar() const
+	{
+		return ImageViewportToolbar;
+	}
+
+	SImageViewport::FPixelCoordinatesUnderCursorResult SImageViewport::GetPixelCoordinatesUnderCursor() const
+	{
+		const auto [PixelCoordsValid, PixelCoords] = ImageViewportClient->GetPixelCoordinatesUnderCursor();
+		return {PixelCoordsValid, PixelCoords};
+	}
+
+	TSharedRef<FEditorViewportClient> SImageViewport::MakeEditorViewportClient()
+	{
+		if (!Client.IsValid())
+		{
+			Client = ImageViewportClient;
+		}
+
+		return Client.ToSharedRef();
+	}
+
+	TSharedPtr<SWidget> SImageViewport::MakeViewportToolbar()
+	{
+		auto HasImage = [this]
+		{
+			if (ImageViewer.IsValid())
+			{
+				const IImageViewer::FImageInfo ImageInfo = ImageViewer->GetCurrentImageInfo();
+				return ImageInfo.bIsValid;
+			}
+			return false;
+		};
+
+		auto GetNumMips = [this]
+		{
+			if (ImageViewer.IsValid())
+			{
+				const IImageViewer::FImageInfo ImageInfo = ImageViewer->GetCurrentImageInfo();
+				if (ImageInfo.bIsValid)
+				{
+					return ImageInfo.NumMips;
+				}
+			}
+
+			return 0;
+		};
+
+		auto GetImageGuid = [this]
+		{
+			return ImageViewer->GetCurrentImageInfo().Guid;
+		};
+
+		SAssignNew(ImageViewportToolbar, SImageViewportToolbar, StaticCastSharedPtr<FImageViewportClient>(Client), CommandList,
+						  SImageViewportToolbar::FConstructParameters
+						  {
+							  SImageViewportToolbar::FHasImage::CreateLambda(HasImage),
+							  SImageViewportToolbar::FNumMips::CreateLambda(GetNumMips),
+							  SImageViewportToolbar::FImageGuid::CreateLambda(GetImageGuid),
+							  ToolbarExtender
+						  });
+
+		ToolbarExtender.Reset();
+
+		return ImageViewportToolbar;
+	}
+
+	void SImageViewport::PopulateViewportOverlays(TSharedRef<SOverlay> Overlay)
+	{
+		Overlay->AddSlot()
+		       .VAlign(VAlign_Bottom)
+		[
+			MakeStatusBar(Overlay)
+		];
+	}
+
+	FText SImageViewport::GetPickerLabel() const
+	{
+		const IImageViewer::FImageInfo ImageInfo = ImageViewer->GetCurrentImageInfo();
+		if (ImageInfo.bIsValid)
+		{
+			const auto [PixelCoordsValid, PixelCoords] = ImageViewportClient->GetPixelCoordinatesUnderCursor();
+			if (PixelCoordsValid)
+			{
+				const FIntPoint PixelCoordsInt(PixelCoords.X, PixelCoords.Y);
+				if (0 <= PixelCoordsInt.X && 0 <= PixelCoordsInt.Y && PixelCoordsInt.X < ImageInfo.Size.X && PixelCoordsInt.Y < ImageInfo.Size.Y)
+				{
+					const int32 MipIndex = FMath::Max(0, ImageViewportClient->GetMipLevel());
+
+					const FIntPoint MipPixelCoords = [&PixelCoordsInt, MipIndex]
+					{
+						FIntPoint Coords = PixelCoordsInt;
+						if (MipIndex > 0)
+						{
+							Coords.X /= 1 << MipIndex;
+							Coords.Y /= 1 << MipIndex;
+						}
+						return Coords;
+					}();
+
+					const TOptional<TVariant<FColor, FLinearColor>> PixelColor = ImageViewer->GetCurrentImagePixelColor(MipPixelCoords, MipIndex);
+
+					auto CoordsAndColorText = [&MipPixelCoords](const auto* Color, const FNumberFormattingOptions& Formatting)
+					{
+						return FText::Format(
+							LOCTEXT("Coordinates", "x={0} y={1}   <RichTextBlock.Red>{2}</> <RichTextBlock.Green>{3}</> <RichTextBlock.Blue>{4}</> {5}"),
+							FText::AsNumber(MipPixelCoords.X), FText::AsNumber(MipPixelCoords.Y),
+							FText::AsNumber(Color->R, &Formatting), FText::AsNumber(Color->G, &Formatting),
+							FText::AsNumber(Color->B, &Formatting), FText::AsNumber(Color->A, &Formatting));
+					};
+
+					FNumberFormattingOptions FormattingByte;
+					FormattingByte.SetMinimumIntegralDigits(3);
+
+					FNumberFormattingOptions FormattingFloat;
+					FormattingFloat.SetMinimumFractionalDigits(3);
+					FormattingFloat.SetMaximumFractionalDigits(3);
+
+					if (PixelColor.IsSet())
+					{
+						const TVariant<FColor, FLinearColor>& PixelColorValue = PixelColor.GetValue();
+						if (const FColor* Color = PixelColorValue.TryGet<FColor>())
+						{
+							return CoordsAndColorText(Color, FormattingByte);
+						}
+						if (const FLinearColor* ColorLinear = PixelColorValue.TryGet<FLinearColor>())
+						{
+							return CoordsAndColorText(ColorLinear, FormattingFloat);
+						}
+					}
+
+					return FText::Format(
+						LOCTEXT("Coordinates", "x={0} y={1}"),
+						FText::AsNumber(MipPixelCoords.X), FText::AsNumber(MipPixelCoords.Y));
+				}
+			}
+		}
+		return {};
+	}
+
+	FText SImageViewport::GetResolutionLabel() const
+	{
+		if (ImageViewer.IsValid())
+		{
+			const IImageViewer::FImageInfo ImageInfo = ImageViewer->GetCurrentImageInfo();
+			if (ImageInfo.bIsValid)
+			{
+				FIntPoint Size = ImageInfo.Size;
+
+				const int32 Mip = ImageViewportClient->GetMipLevel();
+				if (Mip > 0)
+				{
+					Size.X /= 1 << Mip;
+					Size.Y /= 1 << Mip;
+				}
+
+				return FText::Format(LOCTEXT("Resolution", "{0} \u00D7 {1}"), FText::AsNumber(Size.X), FText::AsNumber(Size.Y));
+			}
+		}
+		return {};
+	}
+
+	TSharedRef<SWidget> SImageViewport::MakeStatusBar(TSharedRef<SOverlay> Overlay)
+	{
+		const FMargin SlotPadding(6.0f, 2.0f);
+
+		const TSharedPtr<SHorizontalBox> HorizontalBox = SNew(SHorizontalBox);
+
+		auto ApplyHook = [this, &HorizontalBox](FName ExtensionHook, EExtensionHook::Position HookPosition)
+		{
+			if (ExtensionHook != NAME_None && StatusBarExtender.IsValid())
+			{
+				StatusBarExtender->Apply(ExtensionHook, HookPosition, *HorizontalBox);
+			}
+		};
+
+		ApplyHook(FName("StatusBarLeft"), EExtensionHook::Before);
+		{
+			HorizontalBox->AddSlot()
+			             .Padding(SlotPadding)
+			             .AutoWidth()
+			             .HAlign(HAlign_Left)
+			[
+				SNew(STextBlock)
+				.Text(this, &SImageViewport::GetResolutionLabel)
+			];
+		}
+		ApplyHook(FName("StatusBarLeft"), EExtensionHook::After);
+
+		ApplyHook(FName("StatusBarCenter"), EExtensionHook::Before);
+		{
+			HorizontalBox->AddSlot()
+			             .HAlign(HAlign_Center)
+			[
+				SNullWidget::NullWidget
+
+				// This is deliberately left empty.
+				// It allows for adding toolbar extensions though.
+			];
+		}
+		ApplyHook(FName("StatusBarCenter"), EExtensionHook::After);
+
+		ApplyHook(FName("StatusBarRight"), EExtensionHook::Before);
+		{
+			HorizontalBox->AddSlot()
+			             .Padding(SlotPadding)
+			             .AutoWidth()
+			             .HAlign(HAlign_Right)
+			[
+				SNew(SRichTextBlock)
+					.Text(this, &SImageViewport::GetPickerLabel)
+					.DecoratorStyleSet(&FImageWidgetsStyle::Get())
+			];
+		}
+		ApplyHook(FName("StatusBarRight"), EExtensionHook::After);
+
+		StatusBarExtender.Reset();
+		
+		return HorizontalBox.ToSharedRef();
+	}
+}
+
+#undef LOCTEXT_NAMESPACE
diff --git a/Engine/Plugins/Experimental/ImageWidgets/Source/ImageWidgets/Private/SImageViewportToolbar.cpp b/Engine/Plugins/Experimental/ImageWidgets/Source/ImageWidgets/Private/SImageViewportToolbar.cpp
new file mode 100644
--- /dev/null
+++ b/Engine/Plugins/Experimental/ImageWidgets/Source/ImageWidgets/Private/SImageViewportToolbar.cpp
@@ -0,0 +1,227 @@
+﻿// Copyright Epic Games, Inc. All Rights Reserved.
+
+#include "SImageViewportToolbar.h"
+
+#include "ImageViewportClient.h"
+#include "ImageWidgetsCommands.h"
+#include "SEditorViewportToolBarMenu.h"
+#include "Framework/MultiBox/MultiBoxBuilder.h"
+#include "Settings/LevelEditorViewportSettings.h"
+#include "Widgets/SBoxPanel.h"
+#include "Widgets/Input/SCheckBox.h"
+#include "Widgets/Layout/SBorder.h"
+
+#define LOCTEXT_NAMESPACE "ImageViewportToolbar"
+
+namespace UE::ImageWidgets
+{
+	namespace ImageViewportToolbar_Local
+	{
+		FText Auto = LOCTEXT("Auto", "Auto");
+	}
+
+	void SImageViewportToolbar::Construct(const FArguments& InArgs, const TSharedPtr<FImageViewportClient>& InViewportClient,
+	                                      const TSharedPtr<FUICommandList>& InCommandList, FConstructParameters Parameters)
+	{
+		ViewportClient = InViewportClient;
+		CommandList = InCommandList;
+		check(ViewportClient.IsValid());
+		check(CommandList.IsValid());
+
+		HasImage = MoveTemp(Parameters.HasImage);
+		NumMips = MoveTemp(Parameters.NumMips);
+		ImageGuid = MoveTemp(Parameters.ImageGuid);
+		check(HasImage.IsBound());
+		check(NumMips.IsBound());
+		check(ImageGuid.IsBound());
+
+		ChildSlot
+		[
+			SNew(SBorder)
+			.BorderImage(FAppStyle::GetBrush("EditorViewportToolBar.Background"))
+			[
+				SNew(SHorizontalBox)
+				+ SHorizontalBox::Slot()
+				  .AutoWidth()
+				  .HAlign(HAlign_Left)
+				[
+					MakeLeftToolbar(Parameters.ToolbarExtender)
+				]
+				+ SHorizontalBox::Slot()
+				  .FillWidth(1.0f)
+				  .HAlign(HAlign_Center)
+				[
+					MakeCenterToolbar(Parameters.ToolbarExtender)
+				]
+				+ SHorizontalBox::Slot()
+				  .AutoWidth()
+				  .HAlign(HAlign_Right)
+				[
+					MakeRightToolbar(Parameters.ToolbarExtender)
+				]
+			]
+		];
+
+		SViewportToolBar::Construct(SViewportToolBar::FArguments());
+	}
+
+	FSlimHorizontalToolBarBuilder GetToolbarBuilder(const TSharedPtr<FUICommandList>& CommandList, const TSharedPtr<FExtender>& Extender)
+	{
+		FSlimHorizontalToolBarBuilder ToolbarBuilder(CommandList, FMultiBoxCustomization::None, Extender, true);
+		ToolbarBuilder.SetStyle(&FAppStyle::Get(), "EditorViewportToolBar");
+		ToolbarBuilder.SetIsFocusable(false);
+		ToolbarBuilder.SetLabelVisibility(EVisibility::Collapsed);
+		return ToolbarBuilder;
+	};
+
+	TSharedRef<SWidget> SImageViewportToolbar::MakeLeftToolbar(const TSharedPtr<FExtender>& Extender)
+	{
+		FSlimHorizontalToolBarBuilder ToolbarBuilder = GetToolbarBuilder(CommandList, Extender);
+
+		ToolbarBuilder.BeginSection("ToolbarLeft");
+		{
+			ToolbarBuilder.BeginBlockGroup();
+
+			ToolbarBuilder.AddWidget(
+				SNew(SEditorViewportToolbarMenu)
+				.ParentToolBar(SharedThis(this))
+				.Cursor(EMouseCursor::Default)
+				.Label(this, &SImageViewportToolbar::GetZoomMenuLabel)
+				.OnGetMenuContent(this, &SImageViewportToolbar::MakeZoomMenu)
+				.IsEnabled_Lambda([this] { return HasImage.Execute(); })
+			);
+
+			ToolbarBuilder.AddSeparator();
+
+			ToolbarBuilder.AddWidget(
+				SNew(SEditorViewportToolbarMenu)
+				.ParentToolBar(SharedThis(this))
+				.Cursor(EMouseCursor::Default)
+				.Label(this, &SImageViewportToolbar::GetMipMenuLabel)
+				.OnGetMenuContent(this, &SImageViewportToolbar::MakeMipMenu)
+				.IsEnabled_Lambda([this] { return HasImage.Execute(); })
+				.Visibility(this, &SImageViewportToolbar::GetMipMenuVisibility)
+			);
+
+			ToolbarBuilder.EndBlockGroup();
+		}
+		ToolbarBuilder.EndSection();
+
+		return ToolbarBuilder.MakeWidget();
+	}
+
+	TSharedRef<SWidget> SImageViewportToolbar::MakeCenterToolbar(const TSharedPtr<FExtender>& Extender)
+	{
+		FSlimHorizontalToolBarBuilder ToolbarBuilder = GetToolbarBuilder(CommandList, Extender);
+
+		ToolbarBuilder.BeginSection("ToolbarCenter");
+		{
+			// This is deliberately left empty.
+			// Toolbar extenders use this section to add additional widgets.
+		}
+		ToolbarBuilder.EndSection();
+
+		ToolbarBuilder.AddSeparator();
+
+		return ToolbarBuilder.MakeWidget();
+	}
+
+	TSharedRef<SWidget> SImageViewportToolbar::MakeRightToolbar(const TSharedPtr<FExtender>& Extender)
+	{
+		FSlimHorizontalToolBarBuilder ToolbarBuilder = GetToolbarBuilder(CommandList, Extender);
+
+		ToolbarBuilder.BeginSection("ToolbarRight");
+		{
+			// This is deliberately left empty.
+			// It allows for adding toolbar extensions though.
+		}
+
+		ToolbarBuilder.EndSection();
+
+		return ToolbarBuilder.MakeWidget();
+	}
+
+	FText SImageViewportToolbar::GetZoomMenuLabel() const
+	{
+		const double Zoom = ViewportClient->GetZoom().Zoom;
+		FNumberFormattingOptions FormattingOptions;
+		FormattingOptions.SetMaximumFractionalDigits(Zoom < 1.0 ? (Zoom < 0.1 ? 2 : 1) : 0);
+		return FText::Format(LOCTEXT("Zoom", "Zoom {0}"), FText::AsPercent(Zoom, &FormattingOptions));
+	}
+
+	TSharedRef<SWidget> SImageViewportToolbar::MakeZoomMenu() const
+	{
+		FMenuBuilder MenuBuilder(true, CommandList, nullptr, false, &FAppStyle::Get(), false);
+
+		const auto& Commands = FImageWidgetsCommands::Get();
+
+		MenuBuilder.AddMenuEntry(Commands.Zoom12);
+		MenuBuilder.AddMenuEntry(Commands.Zoom25);
+		MenuBuilder.AddMenuEntry(Commands.Zoom50);
+		MenuBuilder.AddMenuEntry(Commands.Zoom100);
+		MenuBuilder.AddMenuEntry(Commands.Zoom200);
+		MenuBuilder.AddMenuEntry(Commands.Zoom400);
+		MenuBuilder.AddMenuEntry(Commands.Zoom800);
+		MenuBuilder.AddSeparator();
+		MenuBuilder.AddMenuEntry(Commands.ZoomFit);
+		MenuBuilder.AddMenuEntry(Commands.ZoomFill);
+
+		return MenuBuilder.MakeWidget();
+	}
+
+	EVisibility SImageViewportToolbar::GetMipMenuVisibility() const
+	{
+		return NumMips.Execute() > 1 ? EVisibility::Visible : EVisibility::Collapsed;
+	}
+
+	FText SImageViewportToolbar::GetMipMenuLabel() const
+	{
+		const int32 Mip = ViewportClient->GetMipLevel();
+		return FText::Format(LOCTEXT("Mip", "Mip {0}"), Mip == -1 ? ImageViewportToolbar_Local::Auto : FText::AsNumber(Mip));
+	}
+
+	TSharedRef<SWidget> SImageViewportToolbar::MakeMipMenu() const
+	{
+		FMenuBuilder MenuBuilder(true, nullptr, nullptr, false, &FAppStyle::Get(), false);
+
+		auto AddMenuEntry = [this, &MenuBuilder](const TAttribute<FText>& Label, const TAttribute<FText>& ToolTip, int32 MipLevel)
+		{
+			MenuBuilder.AddMenuEntry(
+				Label,
+				ToolTip,
+				FSlateIcon(),
+				FUIAction(
+					FExecuteAction::CreateSP(ViewportClient.ToSharedRef(), &FImageViewportClient::SetMipLevel, MipLevel),
+					FCanExecuteAction(),
+					FIsActionChecked::CreateLambda([this, MipLevel]() { return ViewportClient->GetMipLevel() == MipLevel; })
+				),
+				NAME_None,
+				EUserInterfaceActionType::RadioButton
+			);
+		};
+
+		AddMenuEntry(ImageViewportToolbar_Local::Auto, LOCTEXT("MipLevelAuto_Tooltip", "Choose Mip Level automatically"), -1);
+
+		if (NumMips.IsBound())
+		{
+			const int32 Num = NumMips.Execute();
+			if (Num > 1)
+			{
+				MenuBuilder.AddSeparator();
+
+				for (int32 Mip = 0; Mip < Num; ++Mip)
+				{
+					const FText MipText = FText::AsNumber(Mip);
+
+					AddMenuEntry(FText::Format(LOCTEXT("MipLevel", "Mip {0}"), MipText),
+					             FText::Format(LOCTEXT("MipLevel_Tooltip", "Display Mip Level {0}"), MipText),
+					             Mip);
+				}
+			}
+		}
+
+		return MenuBuilder.MakeWidget();
+	}
+}
+
+#undef LOCTEXT_NAMESPACE
diff --git a/Engine/Plugins/Experimental/ImageWidgets/Source/ImageWidgets/Private/SImageViewportToolbar.h b/Engine/Plugins/Experimental/ImageWidgets/Source/ImageWidgets/Private/SImageViewportToolbar.h
new file mode 100644
--- /dev/null
+++ b/Engine/Plugins/Experimental/ImageWidgets/Source/ImageWidgets/Private/SImageViewportToolbar.h
@@ -0,0 +1,64 @@
+﻿// Copyright Epic Games, Inc. All Rights Reserved.
+
+#pragma once
+
+#include "CoreMinimal.h"
+#include "SImageViewport.h"
+#include "SViewportToolBar.h"
+
+class FExtender;
+class FImageViewportCommands;
+class FUICommandList;
+class SViewportToolBarComboMenu;
+
+namespace UE::ImageWidgets
+{
+	class FImageViewportClient;
+
+	/**
+	 * Extendable toolbar for the image viewport.
+	 */
+	class SImageViewportToolbar : public SViewportToolBar
+	{
+	public:
+
+		/// Callbacks to image viewport to retrieve necessary image metadata.
+		DECLARE_DELEGATE_RetVal(bool, FHasImage)
+		DECLARE_DELEGATE_RetVal(int32, FNumMips)
+		DECLARE_DELEGATE_RetVal(FGuid, FImageGuid)
+		
+		SLATE_BEGIN_ARGS(SImageViewportToolbar)	{}
+		SLATE_END_ARGS()
+
+		/** Parameter for constructing the toolbar. */
+		struct FConstructParameters
+		{
+			FHasImage HasImage;
+			FNumMips NumMips;
+			FImageGuid ImageGuid;
+			TSharedPtr<FExtender> ToolbarExtender;
+		};
+		
+		void Construct(const FArguments& InArgs, const TSharedPtr<FImageViewportClient>& InViewportClient, const TSharedPtr<FUICommandList>& InCommandList,
+		               FConstructParameters Parameters);
+
+	private:
+		TSharedRef<SWidget> MakeLeftToolbar(const TSharedPtr<FExtender>& Extender);
+		TSharedRef<SWidget> MakeCenterToolbar(const TSharedPtr<FExtender>& Extender);
+		TSharedRef<SWidget> MakeRightToolbar(const TSharedPtr<FExtender>& Extender);
+
+		EVisibility GetMipMenuVisibility() const;
+		FText GetMipMenuLabel() const;
+		TSharedRef<SWidget> MakeMipMenu() const;
+
+		FText GetZoomMenuLabel() const;
+		TSharedRef<SWidget> MakeZoomMenu() const;
+
+		TSharedPtr<FImageViewportClient> ViewportClient;
+		TSharedPtr<FUICommandList> CommandList;
+
+		FHasImage HasImage;
+		FNumMips NumMips;
+		FImageGuid ImageGuid;
+	};
+}
diff --git a/Engine/Plugins/Experimental/ImageWidgets/Source/ImageWidgets/Public/IImageViewer.h b/Engine/Plugins/Experimental/ImageWidgets/Source/ImageWidgets/Public/IImageViewer.h
new file mode 100644
--- /dev/null
+++ b/Engine/Plugins/Experimental/ImageWidgets/Source/ImageWidgets/Public/IImageViewer.h
@@ -0,0 +1,102 @@
+// Copyright Epic Games, Inc. All Rights Reserved.
+
+#pragma once
+
+#include "CoreMinimal.h"
+#include "Misc/TVariant.h"
+
+class FCanvas;
+class FViewport;
+
+namespace UE::ImageWidgets
+{
+	/**
+	 * Interface for a component that wants to show image related content with the Slate widgets in the ImageWidget module.
+	 * In this context, and "image" is considered to be any 2D content that is contained within an axis-aligned rectangle.
+	 */
+	class IImageViewer
+	{
+	public:
+		/**
+		 * Information about an image to be displayed.
+		 */
+		struct FImageInfo
+		{
+			/** Unique image identifier. It can encode any helpful metadata as long as no two provided images have the same GUID. */
+			FGuid Guid;
+
+			/** XY size of the image in pixels. */
+			FIntPoint Size;
+
+			/** Number of available MIPs. This should be set to zero if the image type does not support mips. */
+			int32 NumMips;
+
+			/** Indicates that this image is valid for display. */
+			bool bIsValid;
+		};
+
+		/**
+		 * Information necessary for correctly drawing an image.
+		 */
+		struct FDrawProperties
+		{
+			/**
+			 * Where in the 2D plane the image rectangle is supposed to be drawn.
+			 */
+			struct FPlacement
+			{
+				/** Offset from the origin, i.e. (0, 0). */
+				FVector2d Offset;
+
+				/** XY size of the axis aligned rectangle containing the image. */
+				FVector2d Size;
+
+				/**
+				 * The zoom factor used for the image.
+				 * While this might not be necessary for drawing the image, it can be helpful in certain use cases. For example, interpolation could explicitly
+				 * be turned off when zooming into a texture to show the discrete pixel contents of the texture instead of the interpolated result.
+				 */
+				double ZoomFactor;
+			};
+
+			/**
+			 * Information about MIP levels. This can be ignored if the image type does not support MIPs. 
+			 */
+			struct FMip
+			{
+				/** The selected MIP level. */
+				float MipLevel;
+			};
+
+			FPlacement Placement;
+			FMip Mip;
+		};
+
+		/**
+		 * Provide any necessary metadata for the image widgets about the image that is currently supposed to be displayed. This data is generic in the sense
+		 * that the image widgets don't need to know any of the image structure, its content or how to draw it. Instead, the image drawing is done directly
+		 * via @ref DrawCurrentImage.
+		 * @return information about the image to be displayed
+		 */
+		virtual FImageInfo GetCurrentImageInfo() const = 0;
+
+		/**
+		 * Draws the image that is currently supposed to be displayed within the 2D viewport.
+		 * @param Viewport The viewport the image is drawn into
+		 * @param Canvas The canvas used for drawing the image
+		 * @param Properties Information for drawing the image based on the current image viewport state 
+		 */
+		virtual void DrawCurrentImage(FViewport* Viewport, FCanvas* Canvas, const FDrawProperties& Properties) = 0;
+		
+		/**
+		 * Provide information about a given pixel.
+		 * @param PixelCoords XY coordinates for the pixel
+		 * @param MipLevel MIP level that is currently displayed; this can be ignored for images not supporting MIPs 
+		 * @return Either a color value in byte or float format, i.e. FColor or FLinearColor, or no value if there is no valid pixel at the provided coordinates 
+		 */
+		virtual TOptional<TVariant<FColor, FLinearColor>> GetCurrentImagePixelColor(FIntPoint PixelCoords, int32 MipLevel) const = 0;
+
+	protected:
+		~IImageViewer() = default;
+	};
+}
diff --git a/Engine/Plugins/Experimental/ImageWidgets/Source/ImageWidgets/Public/SImageViewport.h b/Engine/Plugins/Experimental/ImageWidgets/Source/ImageWidgets/Public/SImageViewport.h
new file mode 100644
--- /dev/null
+++ b/Engine/Plugins/Experimental/ImageWidgets/Source/ImageWidgets/Public/SImageViewport.h
@@ -0,0 +1,205 @@
+﻿// Copyright Epic Games, Inc. All Rights Reserved.
+
+#pragma once
+
+#include "SEditorViewport.h"
+
+#include "Framework/MultiBox/MultiBoxExtender.h"
+
+class SViewportToolBar;
+
+namespace UE::ImageWidgets
+{
+	class FImageViewportClient;
+	class FStatusBarExtension;
+	class IImageViewer;
+	class SImageViewportToolbar;
+
+	/**
+	 * Generic viewport for displaying and interacting with 2D image-like content.
+	 * The drawing of the images is deferred to an @ref IImageViewer implementation that needs to be provided upon construction. This viewport only uses the 
+	 * metadata provided by the image viewer to have sufficient information about the image without being aware of its actual format or contents.
+	 */
+	class SImageViewport : public SEditorViewport
+	{
+	public:
+		/**
+		 * Contains any extensions widgets for the viewport toolbar.
+		 * This implementation is similar to the existing toolbar extender, i.e. FExtender, which is designed to only work with dedicated toolbar builders and
+		 * can currently not be with status bars.
+		 */
+		class FStatusBarExtender
+		{
+		public:
+			/**
+			 * Delegate that gets called for each extension when the status bar is constructed.
+			 * The given @ref SHorizontalBox is the layout container that the widget for the extension needs to be added to by calling
+			 * @ref SHorizontalBox::AddSlot().
+			 */
+			DECLARE_DELEGATE_OneParam(FDelegate, SHorizontalBox&);
+
+			/**
+			 * Adds an extension to the status bar.
+			 * @param ExtensionHook The section the extension is applied to
+			 * @param HookPosition Where in relation to the section the extension is applied to
+			 * @param Commands The UI command list responsible for handling actions used in the widgets added in the extension
+			 * @param Delegate Delegate called for adding the extension widgets when the status bar is constructed
+			 */
+			IMAGEWIDGETS_API void AddExtension(FName ExtensionHook, EExtensionHook::Position HookPosition, const TSharedPtr<FUICommandList>& Commands,
+			                                   const FDelegate& Delegate);
+
+		private:
+			friend SImageViewport;
+
+			/** Used by the viewport to add extensions to the status bar. */
+			void Apply(FName ExtensionHook, EExtensionHook::Position HookPosition, SHorizontalBox& HorizontalBox) const;
+
+			/** List of extensions that get applied to the viewport status bar. */
+			TArray<TPimplPtr<FStatusBarExtension>> Extensions;
+		};
+
+		/**
+		 * Any settings related to drawing viewport contents other than the image itself. 
+		 */
+		struct FDrawSettings
+		{
+			/** Clear color for the viewport. */
+			FLinearColor ClearColor = FLinearColor::Black;
+
+			/** Flag that enables drawing a border around the image. If enabled, the border is drawn underneath the image.
+			 *  The center of the line drawn for the border is at the exact edge of the image, i.e. when the image is drawn on top of it, the outer half of the
+			 *  border is visible, and the inner half is occluded by the image. */
+			bool bBorderEnabled = false;
+
+			/** Thickness of the border. */
+			float BorderThickness = 1.0f;
+
+			/** Color of the border */
+			FLinearColor BorderColor = FVector3f(0.2f);
+
+			/** Flag that enables the drawing of the background color within the image rectangle. */
+			bool bBackgroundColorEnabled = false;
+
+			/** Color for the background within the image rectangle. If this is different to the clear color, it shows where the image is even if nothing is drawn. */
+			FLinearColor BackgroundColor = FLinearColor::Black;
+
+			/** Flag that enabled the drawing of a background checker texture. If this is enabled, the texture is drawn on top of the background color. */
+			bool bBackgroundCheckerEnabled = false;
+
+			/** First color used in the background checker texture. */
+			FLinearColor BackgroundCheckerColor1 = FLinearColor::White;
+
+			/** Second color used in the background checker texture. */
+			FLinearColor BackgroundCheckerColor2 = FVector3f(0.8f);
+
+			/** Size of a single square within the background checker texture in pixels. */
+			uint32 BackgroundCheckerSize = 8;
+		};
+
+		SLATE_BEGIN_ARGS(SImageViewport)
+			{
+			}
+
+			/** Extensions for the viewport toolbar; valid hooks are "ToolbarLeft", "ToolbarCenter", "ToolbarRight" */
+			SLATE_ARGUMENT(TSharedPtr<FExtender>, ToolbarExtender)
+
+			/** Extensions for the viewport status bar; valid hooks are "StatusBarLeft", "StatusBarCenter", "StatusBarRight" */
+			SLATE_ARGUMENT(TSharedPtr<FStatusBarExtender>, StatusBarExtender)
+
+			/** Settings for drawing viewport contents other than the actual image */
+			SLATE_ATTRIBUTE(FDrawSettings, DrawSettings)
+
+		SLATE_END_ARGS()
+
+		IMAGEWIDGETS_API SImageViewport();
+		IMAGEWIDGETS_API virtual ~SImageViewport() override;
+
+		/**
+		 * Function used by Slate to construct the image viewport widget with the given arguments.
+		 * @param InArgs Slate arguments defined above
+		 * @param InImageViewer @ref IImageViewer implementation that holds and draws the actual image contents.
+		 */
+		IMAGEWIDGETS_API void Construct(const FArguments& InArgs, const TSharedRef<IImageViewer>& InImageViewer);
+
+		/**
+		 * Provides access to the viewport toolbar, which is needed to dynamically generate certain toolbar widgets, e.g. an @ref SEditorViewportToolbarMenu,
+		 * as part of a toolbar extension. 
+		 * @return Pointer to the viewport toolbar 
+		 */
+		IMAGEWIDGETS_API TSharedPtr<SViewportToolBar> GetParentToolbar() const;
+
+		/**
+		 * Result for calls to @ref GetPixelCoordinatesUnderCursor.
+		 */
+		struct FPixelCoordinatesUnderCursorResult
+		{
+			/** Indicates that the cursor position is currently valid. This is set to false, for example, when the cursor is outside of the widget. */
+			bool bIsValid = false;
+
+			/** Pixel coordinates under the cursor relative to the image rectangle size and placement. Note that the coordinates might be outside of the image
+			 * rectangle, i.e. values might be negative or larger than the image size. */
+			FVector2d Coordinates;
+		};
+
+		/**
+		 * Provides the pixel coordinates under the cursor.
+		 * @return Pixel coordinates result under the cursor; see @ref FPixelCoordinatesUnderCursorResult for more details
+		 */
+		IMAGEWIDGETS_API FPixelCoordinatesUnderCursorResult GetPixelCoordinatesUnderCursor() const;
+
+		/**
+		 * Resets the camera controller to default values.
+		 * @param ImageSize Size of the currently displayed image in pixels
+		 */
+		IMAGEWIDGETS_API void ResetController(FIntPoint ImageSize) const;
+
+		/**
+		 * Resets the MIP level to the default.
+		 */
+		IMAGEWIDGETS_API void ResetMip() const;
+
+		/**
+		 * Resets the zoom to default.
+		 * @param ImageSize Size of the currently displayed image in pixels
+		 */
+		IMAGEWIDGETS_API void ResetZoom(FIntPoint ImageSize) const;
+
+		// SEditorViewport overrides - begin
+	public:
+		virtual void BindCommands() override;
+		virtual bool IsVisible() const override { return true; } // Overriding this method avoids the 0.25 seconds update delay in SEditorViewport.
+	protected:
+		virtual TSharedRef<FEditorViewportClient> MakeEditorViewportClient() override;
+		virtual TSharedPtr<SWidget> MakeViewportToolbar() override;
+		virtual void PopulateViewportOverlays(TSharedRef<SOverlay> Overlay) override;
+		// SEditorViewport overrides - end
+
+	private:
+		/** Create the status bar widgets. */
+		TSharedRef<SWidget> MakeStatusBar(TSharedRef<SOverlay> Overlay);
+
+		/** Provides the text for the color picker display in the status bar. */
+		FText GetPickerLabel() const;
+
+		/** Provides the text for the resolution display in the status bar. */
+		FText GetResolutionLabel() const;
+
+		/** Makes the draw setting available either as fixed values or via a callback to the outside of the viewport. */
+		TAttribute<FDrawSettings> DrawSettings;
+
+		/** The image viewer that holds and draws the actual images. */
+		TSharedPtr<IImageViewer> ImageViewer;
+
+		/** The viewport client that takes care of camera controls and displaying the viewport contents. */
+		TSharedPtr<FImageViewportClient> ImageViewportClient;
+
+		/** The toolbar that controls some of the behavior of the viewport and optionally also the image viewer via toolbar extensions. */
+		TSharedPtr<SImageViewportToolbar> ImageViewportToolbar;
+
+		/** Toolbar extensions provided by the call to @ref Construct(). This pointer is reset after the extensions were applied during construction. */
+		TSharedPtr<FExtender> ToolbarExtender;
+
+		/** Status bar extensions provided by the call to @ref Construct(). This pointer is reset after the extensions were applied during construction. */
+		TSharedPtr<FStatusBarExtender> StatusBarExtender;
+	};
+}
