ImageWidgets: Fix localization conflict

diff --git a/Engine/Plugins/Experimental/ImageWidgets/Source/ImageWidgets/Private/SImageViewport.cpp b/Engine/Plugins/Experimental/ImageWidgets/Source/ImageWidgets/Private/SImageViewport.cpp
--- a/Engine/Plugins/Experimental/ImageWidgets/Source/ImageWidgets/Private/SImageViewport.cpp
+++ b/Engine/Plugins/Experimental/ImageWidgets/Source/ImageWidgets/Private/SImageViewport.cpp
@@ -310,7 +310,7 @@ namespace UE::ImageWidgets
 					auto CoordsAndColorText = [&MipPixelCoords](const auto* Color, const FNumberFormattingOptions& Formatting)
 					{
 						return FText::Format(
-							LOCTEXT("Coordinates", "x={0} y={1}   <RichTextBlock.Red>{2}</> <RichTextBlock.Green>{3}</> <RichTextBlock.Blue>{4}</> {5}"),
+							LOCTEXT("CoordinatesWithColor", "x={0} y={1}   <RichTextBlock.Red>{2}</> <RichTextBlock.Green>{3}</> <RichTextBlock.Blue>{4}</> {5}"),
 							FText::AsNumber(MipPixelCoords.X), FText::AsNumber(MipPixelCoords.Y),
 							FText::AsNumber(Color->R, &Formatting), FText::AsNumber(Color->G, &Formatting),
 							FText::AsNumber(Color->B, &Formatting), FText::AsNumber(Color->A, &Formatting));
@@ -337,7 +337,7 @@ namespace UE::ImageWidgets
 					}
 
 					return FText::Format(
-						LOCTEXT("Coordinates", "x={0} y={1}"),
+						LOCTEXT("CoordinatesOnly", "x={0} y={1}"),
 						FText::AsNumber(MipPixelCoords.X), FText::AsNumber(MipPixelCoords.Y));
 				}
 			}
