ImageWidgets: Fix NonUnity build issue

diff --git a/Engine/Plugins/Experimental/ImageWidgets/Source/ImageWidgets/Private/ImageViewportClient.h b/Engine/Plugins/Experimental/ImageWidgets/Source/ImageWidgets/Private/ImageViewportClient.h
--- a/Engine/Plugins/Experimental/ImageWidgets/Source/ImageWidgets/Private/ImageViewportClient.h
+++ b/Engine/Plugins/Experimental/ImageWidgets/Source/ImageWidgets/Private/ImageViewportClient.h
@@ -75,7 +75,7 @@ namespace UE::ImageWidgets
 
 		FImageViewportController Controller;
 
-		TStrongObjectPtr<UTexture2D> CheckerTexture = nullptr;
+		TStrongObjectPtr<UTexture2D> CheckerTexture;
 		FCheckerTextureSettings CachedCheckerTextureSettings;
 	};
 }
