ImageWidgets: Improve SImageViewport zoom settings

diff --git a/Engine/Plugins/Experimental/ImageWidgets/Source/ImageWidgets/ImageWidgets.Build.cs b/Engine/Plugins/Experimental/ImageWidgets/Source/ImageWidgets/ImageWidgets.Build.cs
--- a/Engine/Plugins/Experimental/ImageWidgets/Source/ImageWidgets/ImageWidgets.Build.cs
+++ b/Engine/Plugins/Experimental/ImageWidgets/Source/ImageWidgets/ImageWidgets.Build.cs
@@ -10,8 +10,8 @@ public class ImageWidgets : ModuleRules
 	// Todo Remove build flags as soon as all the respective prototype code is migrated to this module.
 	// These flags are currently used to migrate already existing prototype functionality into this module without having to remove/change the functionality of
 	// the prototype applications.
-	public static bool EnableAbComparison = false;
-	public static bool EnableCatalog = false;
+	private const bool EnableAbComparison = false;
+	private const bool EnableCatalog = false;
 
 	public ImageWidgets(ReadOnlyTargetRules Target) : base(Target)
 	{
diff --git a/Engine/Plugins/Experimental/ImageWidgets/Source/ImageWidgets/Private/ImageViewportClient.cpp b/Engine/Plugins/Experimental/ImageWidgets/Source/ImageWidgets/Private/ImageViewportClient.cpp
--- a/Engine/Plugins/Experimental/ImageWidgets/Source/ImageWidgets/Private/ImageViewportClient.cpp
+++ b/Engine/Plugins/Experimental/ImageWidgets/Source/ImageWidgets/Private/ImageViewportClient.cpp
@@ -98,11 +98,12 @@ void DestroyCheckerTexture(TStrongObjectPtr<UTexture2D>& CheckerTexture)
 }
 
 FImageViewportClient::FImageViewportClient(const TWeakPtr<SEditorViewport>& InEditorViewport, FGetImageSize&& InGetImageSize, FDrawImage&& InDrawImage,
-                                           FGetDrawSettings&& InGetDrawSettings)
+                                           FGetDrawSettings&& InGetDrawSettings, SImageViewport::FControllerSettings::EDefaultZoomMode DefaultZoomMode)
 	: FEditorViewportClient(nullptr, nullptr, InEditorViewport)
 	, GetImageSize(MoveTemp(InGetImageSize))
 	, DrawImage(MoveTemp(InDrawImage))
 	, GetDrawSettings(MoveTemp(InGetDrawSettings))
+	, Controller(static_cast<FImageViewportController::EZoomMode>(DefaultZoomMode))
 {
 	check(GetImageSize.IsBound());
 	check(DrawImage.IsBound());
@@ -290,7 +291,7 @@ FImageViewportController::FZoomSettings FImageViewportClient::GetZoom() const
 {
 	return Controller.GetZoom();
 }
-	
+
 void FImageViewportClient::SetZoom(FImageViewportController::EZoomMode Mode, double Zoom)
 {
 	const FIntPoint ImageSize = GetImageSize.Execute();
diff --git a/Engine/Plugins/Experimental/ImageWidgets/Source/ImageWidgets/Private/ImageViewportClient.h b/Engine/Plugins/Experimental/ImageWidgets/Source/ImageWidgets/Private/ImageViewportClient.h
--- a/Engine/Plugins/Experimental/ImageWidgets/Source/ImageWidgets/Private/ImageViewportClient.h
+++ b/Engine/Plugins/Experimental/ImageWidgets/Source/ImageWidgets/Private/ImageViewportClient.h
@@ -22,7 +22,7 @@ namespace UE::ImageWidgets
 	{
 	public:
 		FImageViewportClient(const TWeakPtr<SEditorViewport>& InViewport, FGetImageSize&& InGetImageSize, FDrawImage&& InDrawImage,
-		                     FGetDrawSettings&& InGetDrawSettings);
+		                     FGetDrawSettings&& InGetDrawSettings, SImageViewport::FControllerSettings::EDefaultZoomMode DefaultZoomMode);
 		virtual ~FImageViewportClient() override;
 
 		virtual void Draw(FViewport* InViewport, FCanvas* Canvas) override;
diff --git a/Engine/Plugins/Experimental/ImageWidgets/Source/ImageWidgets/Private/ImageViewportController.cpp b/Engine/Plugins/Experimental/ImageWidgets/Source/ImageWidgets/Private/ImageViewportController.cpp
--- a/Engine/Plugins/Experimental/ImageWidgets/Source/ImageWidgets/Private/ImageViewportController.cpp
+++ b/Engine/Plugins/Experimental/ImageWidgets/Source/ImageWidgets/Private/ImageViewportController.cpp
@@ -80,7 +80,8 @@ namespace UE::ImageWidgets
 		}
 	}
 
-	FImageViewportController::FImageViewportController()
+	FImageViewportController::FImageViewportController(EZoomMode DefaultZoomMode)
+		:DefaultZoomMode(DefaultZoomMode)
 	{
 		Reset({0, 0}, {0, 0});
 	}
@@ -92,7 +93,7 @@ namespace UE::ImageWidgets
 
 	void FImageViewportController::Reset(FIntPoint ImageSize, FIntPoint ViewportSize)
 	{
-		SetZoom(EZoomMode::Fit, 1.0, ImageSize, ViewportSize);
+		SetZoom(DefaultZoomMode, 1.0, ImageSize, ViewportSize);
 
 		PanAmount = FVector2d::Zero();
 	}
diff --git a/Engine/Plugins/Experimental/ImageWidgets/Source/ImageWidgets/Private/ImageViewportController.h b/Engine/Plugins/Experimental/ImageWidgets/Source/ImageWidgets/Private/ImageViewportController.h
--- a/Engine/Plugins/Experimental/ImageWidgets/Source/ImageWidgets/Private/ImageViewportController.h
+++ b/Engine/Plugins/Experimental/ImageWidgets/Source/ImageWidgets/Private/ImageViewportController.h
@@ -3,6 +3,7 @@
 #pragma once
 
 #include "CoreMinimal.h"
+#include "SImageViewport.h"
 
 namespace UE::ImageWidgets
 {
@@ -12,7 +13,13 @@ namespace UE::ImageWidgets
 	class FImageViewportController
 	{
 	public:
-		enum class EZoomMode { Custom, Fit, Fill };
+		enum class EZoomMode
+		{
+			Fit = static_cast<int32>(SImageViewport::FControllerSettings::EDefaultZoomMode::Fit),
+			Fill = static_cast<int32>(SImageViewport::FControllerSettings::EDefaultZoomMode::Fill),
+			Custom
+		};
+		static_assert(EZoomMode::Fit != EZoomMode::Fill && EZoomMode::Fit != EZoomMode::Custom && EZoomMode::Fill != EZoomMode::Custom);
 
 		struct FZoomSettings
 		{
@@ -20,7 +27,7 @@ namespace UE::ImageWidgets
 			double Zoom;
 		};
 
-		FImageViewportController();
+		explicit FImageViewportController(EZoomMode DefaultZoomMode);
 
 		void Pan(FVector2d Delta);
 		void Reset(FIntPoint ImageSize, FIntPoint ViewportSize);
@@ -34,5 +41,6 @@ namespace UE::ImageWidgets
 	private:
 		FVector2d PanAmount;
 		FZoomSettings ZoomSettings;
+		EZoomMode DefaultZoomMode;
 	};
 }
diff --git a/Engine/Plugins/Experimental/ImageWidgets/Source/ImageWidgets/Private/SImageViewport.cpp b/Engine/Plugins/Experimental/ImageWidgets/Source/ImageWidgets/Private/SImageViewport.cpp
--- a/Engine/Plugins/Experimental/ImageWidgets/Source/ImageWidgets/Private/SImageViewport.cpp
+++ b/Engine/Plugins/Experimental/ImageWidgets/Source/ImageWidgets/Private/SImageViewport.cpp
@@ -172,7 +172,8 @@ namespace UE::ImageWidgets
 			StaticCastWeakPtr<SEditorViewport>(AsWeak()),
 			FGetImageSize::CreateLambda(GetImageSize),
 			FDrawImage::CreateLambda(DrawImage),
-			FGetDrawSettings::CreateLambda(GetDrawSettings)));
+			FGetDrawSettings::CreateLambda(GetDrawSettings),
+			InArgs._ControllerSettings.DefaultZoomMode));
 
 		SEditorViewport::Construct(SEditorViewport::FArguments());
 	}
diff --git a/Engine/Plugins/Experimental/ImageWidgets/Source/ImageWidgets/Private/SImageViewportToolbar.cpp b/Engine/Plugins/Experimental/ImageWidgets/Source/ImageWidgets/Private/SImageViewportToolbar.cpp
--- a/Engine/Plugins/Experimental/ImageWidgets/Source/ImageWidgets/Private/SImageViewportToolbar.cpp
+++ b/Engine/Plugins/Experimental/ImageWidgets/Source/ImageWidgets/Private/SImageViewportToolbar.cpp
@@ -143,10 +143,22 @@ namespace UE::ImageWidgets
 
 	FText SImageViewportToolbar::GetZoomMenuLabel() const
 	{
+		const FImageViewportController::FZoomSettings ZoomSettings = ViewportClient->GetZoom();
+
 		const double Zoom = ViewportClient->GetZoom().Zoom;
 		FNumberFormattingOptions FormattingOptions;
 		FormattingOptions.SetMaximumFractionalDigits(Zoom < 1.0 ? (Zoom < 0.1 ? 2 : 1) : 0);
-		return FText::Format(LOCTEXT("Zoom", "Zoom {0}"), FText::AsPercent(Zoom, &FormattingOptions));
+		const FText ZoomPercentage = FText::AsPercent(Zoom, &FormattingOptions);
+
+		if (ZoomSettings.Mode == FImageViewportController::EZoomMode::Custom)
+		{
+			return ZoomPercentage;
+		}
+
+		return FText::Format(
+			LOCTEXT("ZoomFitFill", "{0} {1}"),
+			FText::FromString(ZoomSettings.Mode == FImageViewportController::EZoomMode::Fit ? "Fit" : "Fill"),
+			ZoomPercentage);
 	}
 
 	TSharedRef<SWidget> SImageViewportToolbar::MakeZoomMenu() const
diff --git a/Engine/Plugins/Experimental/ImageWidgets/Source/ImageWidgets/Public/SImageViewport.h b/Engine/Plugins/Experimental/ImageWidgets/Source/ImageWidgets/Public/SImageViewport.h
--- a/Engine/Plugins/Experimental/ImageWidgets/Source/ImageWidgets/Public/SImageViewport.h
+++ b/Engine/Plugins/Experimental/ImageWidgets/Source/ImageWidgets/Public/SImageViewport.h
@@ -59,7 +59,7 @@ namespace UE::ImageWidgets
 		};
 
 		/**
-		 * Any settings related to drawing viewport contents other than the image itself. 
+		 * Settings related to drawing viewport contents other than the image itself.
 		 */
 		struct FDrawSettings
 		{
@@ -96,6 +96,16 @@ namespace UE::ImageWidgets
 			uint32 BackgroundCheckerSize = 8;
 		};
 
+		/**
+		 * Settings related to viewport controls.
+		 */
+		struct FControllerSettings
+		{
+			enum class EDefaultZoomMode { Fit, Fill };
+
+			EDefaultZoomMode DefaultZoomMode = EDefaultZoomMode::Fit;
+		};
+
 		SLATE_BEGIN_ARGS(SImageViewport)
 			{
 			}
@@ -109,6 +119,9 @@ namespace UE::ImageWidgets
 			/** Settings for drawing viewport contents other than the actual image */
 			SLATE_ATTRIBUTE(FDrawSettings, DrawSettings)
 
+			/** Settings for controlling the viewport */
+			SLATE_ARGUMENT(FControllerSettings, ControllerSettings)
+
 		SLATE_END_ARGS()
 
 		IMAGEWIDGETS_API SImageViewport();
