ImageWidgets: Fix DPI sclaing issues

diff --git a/Engine/Plugins/Experimental/ImageWidgets/Source/ImageWidgets/Private/ImageViewportClient.cpp b/Engine/Plugins/Experimental/ImageWidgets/Source/ImageWidgets/Private/ImageViewportClient.cpp
--- a/Engine/Plugins/Experimental/ImageWidgets/Source/ImageWidgets/Private/ImageViewportClient.cpp
+++ b/Engine/Plugins/Experimental/ImageWidgets/Source/ImageWidgets/Private/ImageViewportClient.cpp
@@ -98,11 +98,13 @@ void DestroyCheckerTexture(TStrongObjectPtr<UTexture2D>& CheckerTexture)
 }
 
 FImageViewportClient::FImageViewportClient(const TWeakPtr<SEditorViewport>& InEditorViewport, FGetImageSize&& InGetImageSize, FDrawImage&& InDrawImage,
-                                           FGetDrawSettings&& InGetDrawSettings, SImageViewport::FControllerSettings::EDefaultZoomMode DefaultZoomMode)
+                                           FGetDrawSettings&& InGetDrawSettings, FGetDPIScaleFactor&& InGetDPIScaleFactor,
+                                           SImageViewport::FControllerSettings::EDefaultZoomMode DefaultZoomMode)
 	: FEditorViewportClient(nullptr, nullptr, InEditorViewport)
 	, GetImageSize(MoveTemp(InGetImageSize))
 	, DrawImage(MoveTemp(InDrawImage))
 	, GetDrawSettings(MoveTemp(InGetDrawSettings))
+	, GetDPIScaleFactor(MoveTemp(InGetDPIScaleFactor))
 	, Controller(static_cast<FImageViewportController::EZoomMode>(DefaultZoomMode))
 {
 	check(GetImageSize.IsBound());
@@ -133,7 +135,8 @@ void FImageViewportClient::Draw(FViewport* InViewport, FCanvas* Canvas)
 		return;
 	}
 
-	CachedPlacement = GetPlacementProperties({InViewport->GetSizeXY()}, {ImageSize});
+	const FVector2d ViewportSizeWithDPIScaling = FVector2d(InViewport->GetSizeXY()) / GetDPIScaleFactor.Execute();	
+	CachedPlacement = GetPlacementProperties(ImageSize, ViewportSizeWithDPIScaling);
 	bCachedPlacementIsValid = true;
 
 	const IImageViewer::FDrawProperties::FMip MipProperties = GetMipProperties();
@@ -164,7 +167,7 @@ void FImageViewportClient::Draw(FViewport* InViewport, FCanvas* Canvas)
 				Background.UV1 = FVector2D(ImageSize) / CheckerTexture->GetSizeX();
 
 				BatchedElementParameters = new FBatchedElementTexture2DPreviewParameters(
-					-1, 0, 0, false, false, false, false, false, false/*MipProperties.bUsePointSampling*/);
+					-1, 0, 0, false, false, false, false, false, false);
 				Background.BatchedElementParameters = BatchedElementParameters.GetReference();
 			}
 		}
@@ -221,9 +224,7 @@ bool FImageViewportClient::InputKey(const FInputKeyEventArgs& EventArgs)
 
 		if (EventArgs.Key == EKeys::F)
 		{
-			const FIntPoint ImageSize = GetImageSize.Execute();
-			const FIntPoint ViewportSize = Viewport->GetSizeXY();
-			Controller.Reset(ImageSize, ViewportSize);
+			Controller.Reset(GetImageSize.Execute(), GetViewportSizeWithDPIScaling());
 
 			return true;
 		}
@@ -258,7 +259,9 @@ void FImageViewportClient::TrackingStopped()
 
 		FIntPoint DraggingEnd;
 		Viewport->GetMousePos(DraggingEnd);
-		Controller.Pan(DraggingEnd - DraggingStart);
+		const float DPIScaleFactor = GetDPIScaleFactor.Execute();
+
+		Controller.Pan(FVector2d(DraggingEnd - DraggingStart) / DPIScaleFactor);
 	}
 
 	RequiredCursorVisibiltyAndAppearance.bDontResetCursor = true;
@@ -278,26 +281,25 @@ void FImageViewportClient::SetMipLevel(int32 InMipLevel)
 
 void FImageViewportClient::ResetController(FIntPoint ImageSize)
 {
-	Controller.Reset(ImageSize, Viewport->GetSizeXY());
+	Controller.Reset(ImageSize, GetViewportSizeWithDPIScaling());
 }
 
 void FImageViewportClient::ResetZoom(FIntPoint ImageSize)
 {		
 	const FImageViewportController::FZoomSettings ZoomSettings = Controller.GetZoom();
-	Controller.SetZoom(ZoomSettings.Mode, ZoomSettings.Zoom, ImageSize, Viewport->GetSizeXY());
+	Controller.SetZoom(ZoomSettings.Mode, ZoomSettings.Zoom, ImageSize, GetViewportSizeWithDPIScaling());
 }
 
 FImageViewportController::FZoomSettings FImageViewportClient::GetZoom() const
 {
-	return Controller.GetZoom();
+	FImageViewportController::FZoomSettings ZoomSettings = Controller.GetZoom();
+	ZoomSettings.Zoom *= GetDPIScaleFactor.Execute();
+	return ZoomSettings;
 }
 
 void FImageViewportClient::SetZoom(FImageViewportController::EZoomMode Mode, double Zoom)
 {
-	const FIntPoint ImageSize = GetImageSize.Execute();
-	const FIntPoint ViewportSize = Viewport->GetSizeXY();
-
-	Controller.SetZoom(Mode, Zoom, ImageSize, ViewportSize);
+	Controller.SetZoom(Mode, Zoom, GetImageSize.Execute(), GetViewportSizeWithDPIScaling());
 }
 
 TPair<bool, FVector2d> FImageViewportClient::GetPixelCoordinatesUnderCursor() const
@@ -307,8 +309,9 @@ TPair<bool, FVector2d> FImageViewportClient::GetPixelCoordinatesUnderCursor() co
 		return {false, FVector2d::Zero()};
 	}
 
-	const FVector2d MousePos(CurrentMousePos.X + 0.5, CurrentMousePos.Y + 0.5);
-	const FVector2d CurrentDrag = GetCurrentDrag();
+	const float DPIScaleFactor = GetDPIScaleFactor.Execute();
+	const FVector2d MousePos((CurrentMousePos.X + 0.5) / DPIScaleFactor, (CurrentMousePos.Y + 0.5) / DPIScaleFactor);
+	const FVector2d CurrentDrag = GetCurrentDragWithDPIScaling();
 
 	FVector2d CursorPos = FVector2d(MousePos) - CachedPlacement.Offset + CurrentDrag;
 	CursorPos /= CachedPlacement.ZoomFactor;
@@ -316,26 +319,30 @@ TPair<bool, FVector2d> FImageViewportClient::GetPixelCoordinatesUnderCursor() co
 	return {true, CursorPos};
 }
 
-FVector2d FImageViewportClient::GetCurrentDrag() const
+FVector2d FImageViewportClient::GetCurrentDragWithDPIScaling() const
 {
 	if (bDragging)
 	{
 		FIntPoint DraggingEnd;
 		Viewport->GetMousePos(DraggingEnd);
-		return DraggingEnd - DraggingStart;
+
+		const float DPIScaleFactor = GetDPIScaleFactor.Execute();
+
+		return FVector2d(DraggingEnd - DraggingStart) / DPIScaleFactor;
 	}
 	return FVector2d::Zero();;
 }
 
-IImageViewer::FDrawProperties::FPlacement FImageViewportClient::GetPlacementProperties(FVector2d ViewportSize, FVector2d ImageSize) const
+IImageViewer::FDrawProperties::FPlacement FImageViewportClient::GetPlacementProperties(const FIntPoint ImageSize,
+                                                                                       const FVector2d ViewportSizeWithDPIScaling) const
 {
-	const FVector2d CurrentDrag = GetCurrentDrag();
+	const FVector2d CurrentDrag = GetCurrentDragWithDPIScaling();
 	const FVector2d Pan = Controller.GetPan(CurrentDrag);
 
 	const double ZoomFactor = Controller.GetZoom().Zoom;
 
-	const FVector2d TileSize = ImageSize * ZoomFactor;
-	const FVector2d TileOffset = (ViewportSize - TileSize) / 2 + Pan;
+	const FVector2d TileSize = FVector2d(ImageSize) * ZoomFactor;
+	const FVector2d TileOffset = (ViewportSizeWithDPIScaling - TileSize) / 2 + Pan;
 
 	return {TileOffset, TileSize, ZoomFactor};
 }
@@ -372,6 +379,14 @@ void FImageViewportClient::CreateOrDestroyCheckerTextureIfSettingsChanged(const
 		CachedCheckerTextureSettings = NewCheckerTextureSettings;
 	}
 }
+
+FVector2d FImageViewportClient::GetViewportSizeWithDPIScaling() const
+{
+	const FIntPoint ViewportSize = Viewport->GetSizeXY();
+	const float DPIScaleFactor = GetDPIScaleFactor.Execute();
+
+	return FVector2d(ViewportSize) / DPIScaleFactor;
+}
 }
 
 #undef LOCTEXT_NAMESPACE
diff --git a/Engine/Plugins/Experimental/ImageWidgets/Source/ImageWidgets/Private/ImageViewportClient.h b/Engine/Plugins/Experimental/ImageWidgets/Source/ImageWidgets/Private/ImageViewportClient.h
--- a/Engine/Plugins/Experimental/ImageWidgets/Source/ImageWidgets/Private/ImageViewportClient.h
+++ b/Engine/Plugins/Experimental/ImageWidgets/Source/ImageWidgets/Private/ImageViewportClient.h
@@ -14,6 +14,7 @@ namespace UE::ImageWidgets
 	DECLARE_DELEGATE_RetVal(FIntPoint, FGetImageSize)
 	DECLARE_DELEGATE_ThreeParams(FDrawImage, FViewport*, FCanvas*, const IImageViewer::FDrawProperties&)
 	DECLARE_DELEGATE_RetVal(SImageViewport::FDrawSettings, FGetDrawSettings)
+	DECLARE_DELEGATE_RetVal(float, FGetDPIScaleFactor)
 
 	/**
 	 * Viewport client for controlling the camera and drawing viewport contents. 
@@ -22,7 +23,8 @@ namespace UE::ImageWidgets
 	{
 	public:
 		FImageViewportClient(const TWeakPtr<SEditorViewport>& InViewport, FGetImageSize&& InGetImageSize, FDrawImage&& InDrawImage,
-		                     FGetDrawSettings&& InGetDrawSettings, SImageViewport::FControllerSettings::EDefaultZoomMode DefaultZoomMode);
+		                     FGetDrawSettings&& InGetDrawSettings, FGetDPIScaleFactor&& InGetDPIScaleFactor,
+		                     SImageViewport::FControllerSettings::EDefaultZoomMode DefaultZoomMode);
 		virtual ~FImageViewportClient() override;
 
 		virtual void Draw(FViewport* InViewport, FCanvas* Canvas) override;
@@ -54,15 +56,18 @@ namespace UE::ImageWidgets
 			uint32 CheckerSize = 0;
 		};
 
-		FVector2d GetCurrentDrag() const;
-		IImageViewer::FDrawProperties::FPlacement GetPlacementProperties(FVector2d ViewportSize, FVector2d TextureSize) const;
+		FVector2d GetCurrentDragWithDPIScaling() const;
+		IImageViewer::FDrawProperties::FPlacement GetPlacementProperties(FIntPoint ImageSize, FVector2d ViewportSizeWithDPIScaling) const;
 		IImageViewer::FDrawProperties::FMip GetMipProperties() const;
 
 		void CreateOrDestroyCheckerTextureIfSettingsChanged(const SImageViewport::FDrawSettings& DrawSettings);
 
+		FVector2d GetViewportSizeWithDPIScaling() const;
+
 		FGetImageSize GetImageSize;
 		FDrawImage DrawImage;
 		FGetDrawSettings GetDrawSettings;
+		FGetDPIScaleFactor GetDPIScaleFactor;
 
 		bool bDragging = false;
 		FIntPoint DraggingStart;
diff --git a/Engine/Plugins/Experimental/ImageWidgets/Source/ImageWidgets/Private/ImageViewportController.cpp b/Engine/Plugins/Experimental/ImageWidgets/Source/ImageWidgets/Private/ImageViewportController.cpp
--- a/Engine/Plugins/Experimental/ImageWidgets/Source/ImageWidgets/Private/ImageViewportController.cpp
+++ b/Engine/Plugins/Experimental/ImageWidgets/Source/ImageWidgets/Private/ImageViewportController.cpp
@@ -80,45 +80,45 @@ namespace UE::ImageWidgets
 		}
 	}
 
-	FImageViewportController::FImageViewportController(EZoomMode DefaultZoomMode)
-		:DefaultZoomMode(DefaultZoomMode)
+	FImageViewportController::FImageViewportController(const EZoomMode DefaultZoomMode)
+		: DefaultZoomMode(DefaultZoomMode)
 	{
 		Reset({0, 0}, {0, 0});
 	}
 
-	void FImageViewportController::Pan(FVector2d Delta)
+	void FImageViewportController::Pan(const FVector2d ViewportDeltaWithDPIScaling)
 	{
-		PanAmount += Delta;
+		PanAmount += ViewportDeltaWithDPIScaling;
 	}
 
-	void FImageViewportController::Reset(FIntPoint ImageSize, FIntPoint ViewportSize)
+	void FImageViewportController::Reset(const FIntPoint ImageSize, const FVector2d ViewportSizeWithDPIScaling)
 	{
-		SetZoom(DefaultZoomMode, 1.0, ImageSize, ViewportSize);
+		SetZoom(DefaultZoomMode, 1.0, ImageSize, ViewportSizeWithDPIScaling);
 
 		PanAmount = FVector2d::Zero();
 	}
 
-	void FImageViewportController::ZoomIn(FVector2d CursorPos, const FIntPoint& ImageSize)
+	void FImageViewportController::ZoomIn(const FVector2d CursorPixelPos, const FIntPoint ImageSize)
 	{
 		const double OldZoom = ZoomSettings.Zoom;
 
 		ZoomSettings.Mode = EZoomMode::Custom;
 		ZoomSettings.Zoom = ImageViewportController_Local::ZoomIn(ZoomSettings.Zoom);
 
-		PanAmount = ImageViewportController_Local::AdjustPan(PanAmount, OldZoom, ZoomSettings.Zoom, CursorPos, FVector2d(ImageSize.X, ImageSize.Y));
+		PanAmount = ImageViewportController_Local::AdjustPan(PanAmount, OldZoom, ZoomSettings.Zoom, CursorPixelPos, ImageSize);
 	}
 
-	void FImageViewportController::ZoomOut(FVector2d CursorPos, const FIntPoint& ImageSize)
+	void FImageViewportController::ZoomOut(const FVector2d CursorPixelPos, const FIntPoint ImageSize)
 	{
 		const double OldZoom = ZoomSettings.Zoom;
 
 		ZoomSettings.Mode = EZoomMode::Custom;
 		ZoomSettings.Zoom = ImageViewportController_Local::ZoomOut(ZoomSettings.Zoom);
 
-		PanAmount = ImageViewportController_Local::AdjustPan(PanAmount, OldZoom, ZoomSettings.Zoom, CursorPos, FVector2d(ImageSize.X, ImageSize.Y));
+		PanAmount = ImageViewportController_Local::AdjustPan(PanAmount, OldZoom, ZoomSettings.Zoom, CursorPixelPos, ImageSize);
 	}
 
-	FVector2d FImageViewportController::GetPan(FVector2d Drag) const
+	FVector2d FImageViewportController::GetPan(const FVector2d Drag) const
 	{
 		return PanAmount + Drag;
 	}
@@ -128,7 +128,7 @@ namespace UE::ImageWidgets
 		return ZoomSettings;
 	}
 
-	void FImageViewportController::SetZoom(const EZoomMode ZoomMode, const double Zoom, const FIntPoint& ImageSize, const FIntPoint& ViewportSize)
+	void FImageViewportController::SetZoom(const EZoomMode ZoomMode, const double Zoom, const FIntPoint ImageSize, const FVector2d ViewportSizeWithDPIScaling)
 	{
 		ZoomSettings.Mode = ZoomMode;
 
@@ -138,13 +138,13 @@ namespace UE::ImageWidgets
 			return;
 		}
 
-		if (ImageSize == FIntPoint::ZeroValue || ViewportSize == FIntPoint::ZeroValue)
+		if (ImageSize == FIntPoint::ZeroValue || ViewportSizeWithDPIScaling == FVector2d::Zero())
 		{
 			ZoomSettings.Zoom = 1.0;
 			return;
 		}
 
-		const FVector2d SizeRatio = FVector2d(ViewportSize) / ImageSize;
+		const FVector2d SizeRatio = ViewportSizeWithDPIScaling / ImageSize;
 		const double SizeRatioMin = SizeRatio.GetMin();
 
 		if (ZoomSettings.Mode == EZoomMode::Fill)
diff --git a/Engine/Plugins/Experimental/ImageWidgets/Source/ImageWidgets/Private/ImageViewportController.h b/Engine/Plugins/Experimental/ImageWidgets/Source/ImageWidgets/Private/ImageViewportController.h
--- a/Engine/Plugins/Experimental/ImageWidgets/Source/ImageWidgets/Private/ImageViewportController.h
+++ b/Engine/Plugins/Experimental/ImageWidgets/Source/ImageWidgets/Private/ImageViewportController.h
@@ -9,6 +9,8 @@ namespace UE::ImageWidgets
 {
 	/**
 	 * Camera controller for the 2D viewport supporting panning and zooming.
+	 *
+	 * Note: All function parameters relating to the viewport need to be passed in with DPI scaling taken into account!
 	 */
 	class FImageViewportController
 	{
@@ -29,14 +31,14 @@ namespace UE::ImageWidgets
 
 		explicit FImageViewportController(EZoomMode DefaultZoomMode);
 
-		void Pan(FVector2d Delta);
-		void Reset(FIntPoint ImageSize, FIntPoint ViewportSize);
-		void ZoomIn(FVector2d CursorPos, const FIntPoint& ImageSize);
-		void ZoomOut(FVector2d CursorPos, const FIntPoint& ImageSize);
+		void Pan(FVector2d ViewportDeltaWithDPIScaling);
+		void Reset(FIntPoint ImageSize, FVector2d ViewportSizeWithDPIScaling);
+		void ZoomIn(FVector2d CursorPixelPos, FIntPoint ImageSize);
+		void ZoomOut(FVector2d CursorPixelPos, FIntPoint ImageSize);
 
 		FVector2d GetPan(FVector2d Drag) const;
 		FZoomSettings GetZoom() const;
-		void SetZoom(EZoomMode ZoomMode, double Zoom, const FIntPoint& ImageSize, const FIntPoint& ViewportSize);
+		void SetZoom(EZoomMode ZoomMode, double Zoom, FIntPoint ImageSize, FVector2d ViewportSizeWithDPIScaling);
 
 	private:
 		FVector2d PanAmount;
diff --git a/Engine/Plugins/Experimental/ImageWidgets/Source/ImageWidgets/Private/SImageViewport.cpp b/Engine/Plugins/Experimental/ImageWidgets/Source/ImageWidgets/Private/SImageViewport.cpp
--- a/Engine/Plugins/Experimental/ImageWidgets/Source/ImageWidgets/Private/SImageViewport.cpp
+++ b/Engine/Plugins/Experimental/ImageWidgets/Source/ImageWidgets/Private/SImageViewport.cpp
@@ -8,6 +8,7 @@
 #include "ImageWidgetsCommands.h"
 #include "ImageWidgetsStyle.h"
 #include "SImageViewportToolbar.h"
+#include "Framework/Application/SlateApplication.h"
 #include "Widgets/Text/SRichTextBlock.h"
 #include "Widgets/Text/STextBlock.h"
 
@@ -168,11 +169,21 @@ namespace UE::ImageWidgets
 			return DrawSettings.Get();
 		};
 
+		auto GetDPIScaleFactor = [this]
+		{
+			if (const TSharedPtr<SWindow> TopLevelWindow = FSlateApplication::Get().FindWidgetWindow(SharedThis(this)); TopLevelWindow.IsValid())
+			{
+				return TopLevelWindow->GetDPIScaleFactor();
+			}
+			return 1.0f;
+		};
+		
 		ImageViewportClient = MakeShareable(new FImageViewportClient(
 			StaticCastWeakPtr<SEditorViewport>(AsWeak()),
 			FGetImageSize::CreateLambda(GetImageSize),
 			FDrawImage::CreateLambda(DrawImage),
 			FGetDrawSettings::CreateLambda(GetDrawSettings),
+			FGetDPIScaleFactor::CreateLambda(GetDPIScaleFactor),
 			InArgs._ControllerSettings.DefaultZoomMode));
 
 		SEditorViewport::Construct(SEditorViewport::FArguments());
diff --git a/Engine/Plugins/Experimental/ImageWidgets/Source/ImageWidgets/Private/SImageViewportToolbar.cpp b/Engine/Plugins/Experimental/ImageWidgets/Source/ImageWidgets/Private/SImageViewportToolbar.cpp
--- a/Engine/Plugins/Experimental/ImageWidgets/Source/ImageWidgets/Private/SImageViewportToolbar.cpp
+++ b/Engine/Plugins/Experimental/ImageWidgets/Source/ImageWidgets/Private/SImageViewportToolbar.cpp
@@ -145,7 +145,7 @@ namespace UE::ImageWidgets
 	{
 		const FImageViewportController::FZoomSettings ZoomSettings = ViewportClient->GetZoom();
 
-		const double Zoom = ViewportClient->GetZoom().Zoom;
+		const double Zoom = ZoomSettings.Zoom;
 		FNumberFormattingOptions FormattingOptions;
 		FormattingOptions.SetMaximumFractionalDigits(Zoom < 1.0 ? (Zoom < 0.1 ? 2 : 1) : 0);
 		const FText ZoomPercentage = FText::AsPercent(Zoom, &FormattingOptions);
