The painting editor now has a left mouse button pressed/released events. Events are being bubbled up from the ImageWidget module.

diff --git a/Engine/Plugins/Experimental/ImageWidgets/Source/ImageWidgets/Private/ImageViewportClient.cpp b/Engine/Plugins/Experimental/ImageWidgets/Source/ImageWidgets/Private/ImageViewportClient.cpp
--- a/Engine/Plugins/Experimental/ImageWidgets/Source/ImageWidgets/Private/ImageViewportClient.cpp
+++ b/Engine/Plugins/Experimental/ImageWidgets/Source/ImageWidgets/Private/ImageViewportClient.cpp
@@ -99,17 +99,21 @@ void DestroyCheckerTexture(TStrongObjectPtr<UTexture2D>& CheckerTexture)
 
 FImageViewportClient::FImageViewportClient(const TWeakPtr<SEditorViewport>& InEditorViewport, FGetImageSize&& InGetImageSize, FDrawImage&& InDrawImage,
                                            FGetDrawSettings&& InGetDrawSettings, FGetDPIScaleFactor&& InGetDPIScaleFactor,
-                                           SImageViewport::FControllerSettings::EDefaultZoomMode DefaultZoomMode)
+                                           FOnLeftMouseButtonPressed&& InOnLeftMouseButtonPressed, FOnLeftMouseButtonReleased&& InOnLeftMouseButtonReleased, SImageViewport::FControllerSettings::EDefaultZoomMode DefaultZoomMode)
 	: FEditorViewportClient(nullptr, nullptr, InEditorViewport)
 	, GetImageSize(MoveTemp(InGetImageSize))
 	, DrawImage(MoveTemp(InDrawImage))
 	, GetDrawSettings(MoveTemp(InGetDrawSettings))
 	, GetDPIScaleFactor(MoveTemp(InGetDPIScaleFactor))
+	, OnLeftMouseButtonPressed(MoveTemp(InOnLeftMouseButtonPressed))
+	, OnLeftMouseButtonReleased(MoveTemp(InOnLeftMouseButtonReleased))
 	, Controller(static_cast<FImageViewportController::EZoomMode>(DefaultZoomMode))
 {
 	check(GetImageSize.IsBound());
 	check(DrawImage.IsBound());
 	check(GetDrawSettings.IsBound());
+	check(OnLeftMouseButtonPressed.IsBound());
+	check(OnLeftMouseButtonReleased.IsBound());
 
 	SetRealtime(true);
 }
@@ -230,6 +234,18 @@ bool FImageViewportClient::InputKey(const FInputKeyEventArgs& EventArgs)
 		}
 	}
 
+	if (EventArgs.Key == EKeys::LeftMouseButton)
+	{
+		if(EventArgs.Event == IE_Pressed)
+		{
+			OnLeftMouseButtonPressed.Execute();
+		}
+		else if(EventArgs.Event == IE_Released)
+		{
+			OnLeftMouseButtonReleased.Execute();
+		}
+	}
+
 	return FEditorViewportClient::InputKey(EventArgs);
 }
 
diff --git a/Engine/Plugins/Experimental/ImageWidgets/Source/ImageWidgets/Private/ImageViewportClient.h b/Engine/Plugins/Experimental/ImageWidgets/Source/ImageWidgets/Private/ImageViewportClient.h
--- a/Engine/Plugins/Experimental/ImageWidgets/Source/ImageWidgets/Private/ImageViewportClient.h
+++ b/Engine/Plugins/Experimental/ImageWidgets/Source/ImageWidgets/Private/ImageViewportClient.h
@@ -15,6 +15,8 @@ namespace UE::ImageWidgets
 	DECLARE_DELEGATE_ThreeParams(FDrawImage, FViewport*, FCanvas*, const IImageViewer::FDrawProperties&)
 	DECLARE_DELEGATE_RetVal(SImageViewport::FDrawSettings, FGetDrawSettings)
 	DECLARE_DELEGATE_RetVal(float, FGetDPIScaleFactor)
+	DECLARE_DELEGATE(FOnLeftMouseButtonPressed);
+	DECLARE_DELEGATE(FOnLeftMouseButtonReleased);
 
 	/**
 	 * Viewport client for controlling the camera and drawing viewport contents. 
@@ -24,7 +26,7 @@ namespace UE::ImageWidgets
 	public:
 		FImageViewportClient(const TWeakPtr<SEditorViewport>& InViewport, FGetImageSize&& InGetImageSize, FDrawImage&& InDrawImage,
 		                     FGetDrawSettings&& InGetDrawSettings, FGetDPIScaleFactor&& InGetDPIScaleFactor,
-		                     SImageViewport::FControllerSettings::EDefaultZoomMode DefaultZoomMode);
+		                     FOnLeftMouseButtonPressed&& InOnLeftMouseButtonPressed, FOnLeftMouseButtonReleased&& InOnLeftMouseButtonReleased, SImageViewport::FControllerSettings::EDefaultZoomMode DefaultZoomMode);
 		virtual ~FImageViewportClient() override;
 
 		virtual void Draw(FViewport* InViewport, FCanvas* Canvas) override;
@@ -63,12 +65,14 @@ namespace UE::ImageWidgets
 		void CreateOrDestroyCheckerTextureIfSettingsChanged(const SImageViewport::FDrawSettings& DrawSettings);
 
 		FVector2d GetViewportSizeWithDPIScaling() const;
-
+		
 		FGetImageSize GetImageSize;
 		FDrawImage DrawImage;
 		FGetDrawSettings GetDrawSettings;
 		FGetDPIScaleFactor GetDPIScaleFactor;
-
+		FOnLeftMouseButtonPressed OnLeftMouseButtonPressed;
+		FOnLeftMouseButtonReleased OnLeftMouseButtonReleased;
+		
 		bool bDragging = false;
 		FIntPoint DraggingStart;
 
diff --git a/Engine/Plugins/Experimental/ImageWidgets/Source/ImageWidgets/Private/SImageViewport.cpp b/Engine/Plugins/Experimental/ImageWidgets/Source/ImageWidgets/Private/SImageViewport.cpp
--- a/Engine/Plugins/Experimental/ImageWidgets/Source/ImageWidgets/Private/SImageViewport.cpp
+++ b/Engine/Plugins/Experimental/ImageWidgets/Source/ImageWidgets/Private/SImageViewport.cpp
@@ -148,6 +148,8 @@ namespace UE::ImageWidgets
 		StatusBarExtender = InArgs._StatusBarExtender;
 
 		DrawSettings = InArgs._DrawSettings;
+		OnLeftMouseButtonPressed = InArgs._OnLeftMouseButtonPressed;
+		OnLeftMouseButtonReleased = InArgs._OnLeftMouseButtonReleased;
 
 		auto GetImageSize = [this]
 		{
@@ -177,6 +179,16 @@ namespace UE::ImageWidgets
 			}
 			return 1.0f;
 		};
+
+		auto HandleOnLeftMouseButtonPressed = [this]()
+		{
+			OnLeftMouseButtonPressed.ExecuteIfBound();
+		};
+
+		auto HandleOnLeftMouseButtonReleased = [this]()
+		{
+			OnLeftMouseButtonReleased.ExecuteIfBound();
+		};
 		
 		ImageViewportClient = MakeShareable(new FImageViewportClient(
 			StaticCastWeakPtr<SEditorViewport>(AsWeak()),
@@ -184,7 +196,10 @@ namespace UE::ImageWidgets
 			FDrawImage::CreateLambda(DrawImage),
 			FGetDrawSettings::CreateLambda(GetDrawSettings),
 			FGetDPIScaleFactor::CreateLambda(GetDPIScaleFactor),
-			InArgs._ControllerSettings.DefaultZoomMode));
+			FOnLeftMouseButtonPressed::CreateLambda(HandleOnLeftMouseButtonPressed),
+			FOnLeftMouseButtonReleased::CreateLambda(HandleOnLeftMouseButtonReleased),
+			InArgs._ControllerSettings.DefaultZoomMode)
+			);
 
 		SEditorViewport::Construct(SEditorViewport::FArguments());
 	}
diff --git a/Engine/Plugins/Experimental/ImageWidgets/Source/ImageWidgets/Public/SImageViewport.h b/Engine/Plugins/Experimental/ImageWidgets/Source/ImageWidgets/Public/SImageViewport.h
--- a/Engine/Plugins/Experimental/ImageWidgets/Source/ImageWidgets/Public/SImageViewport.h
+++ b/Engine/Plugins/Experimental/ImageWidgets/Source/ImageWidgets/Public/SImageViewport.h
@@ -14,7 +14,7 @@ namespace UE::ImageWidgets
 	class FStatusBarExtension;
 	class IImageViewer;
 	class SImageViewportToolbar;
-
+	
 	/**
 	 * Generic viewport for displaying and interacting with 2D image-like content.
 	 * The drawing of the images is deferred to an @ref IImageViewer implementation that needs to be provided upon construction. This viewport only uses the 
@@ -106,22 +106,29 @@ namespace UE::ImageWidgets
 			EDefaultZoomMode DefaultZoomMode = EDefaultZoomMode::Fit;
 		};
 
-		SLATE_BEGIN_ARGS(SImageViewport)
-			{
-			}
+		DECLARE_DELEGATE(FOnLeftMouseButtonPressed);
+		DECLARE_DELEGATE(FOnLeftMouseButtonReleased);
+
+		SLATE_BEGIN_ARGS(SImageViewport){}
 
-			/** Extensions for the viewport toolbar; valid hooks are "ToolbarLeft", "ToolbarCenter", "ToolbarRight" */
-			SLATE_ARGUMENT(TSharedPtr<FExtender>, ToolbarExtender)
+		/** Extensions for the viewport toolbar; valid hooks are "ToolbarLeft", "ToolbarCenter", "ToolbarRight" */
+		SLATE_ARGUMENT(TSharedPtr<FExtender>, ToolbarExtender)
 
-			/** Extensions for the viewport status bar; valid hooks are "StatusBarLeft", "StatusBarCenter", "StatusBarRight" */
-			SLATE_ARGUMENT(TSharedPtr<FStatusBarExtender>, StatusBarExtender)
+		/** Extensions for the viewport status bar; valid hooks are "StatusBarLeft", "StatusBarCenter", "StatusBarRight" */
+		SLATE_ARGUMENT(TSharedPtr<FStatusBarExtender>, StatusBarExtender)
 
-			/** Settings for drawing viewport contents other than the actual image */
-			SLATE_ATTRIBUTE(FDrawSettings, DrawSettings)
+		/** Settings for drawing viewport contents other than the actual image */
+		SLATE_ATTRIBUTE(FDrawSettings, DrawSettings)
 
-			/** Settings for controlling the viewport */
-			SLATE_ARGUMENT(FControllerSettings, ControllerSettings)
+		/** Settings for controlling the viewport */
+		SLATE_ARGUMENT(FControllerSettings, ControllerSettings)
 
+		/** Left mouse button pressed event bubbled up from the viewport */
+		SLATE_EVENT(FOnLeftMouseButtonPressed, OnLeftMouseButtonPressed)
+
+		/** Left mouse button released event bubbled up from the viewport */
+		SLATE_EVENT(FOnLeftMouseButtonReleased, OnLeftMouseButtonReleased)
+		
 		SLATE_END_ARGS()
 
 		IMAGEWIDGETS_API SImageViewport();
@@ -214,5 +221,9 @@ namespace UE::ImageWidgets
 
 		/** Status bar extensions provided by the call to @ref Construct(). This pointer is reset after the extensions were applied during construction. */
 		TSharedPtr<FStatusBarExtender> StatusBarExtender;
+
+		/** Left Mouse buttons pressed and released events triggered to tell the widget containing the viewport. These are initialized from the constructor.*/
+		FOnLeftMouseButtonPressed OnLeftMouseButtonPressed;
+		FOnLeftMouseButtonReleased OnLeftMouseButtonReleased;
 	};
 }
