ImageWidgets: Add SImageCatalog

diff --git a/Engine/Plugins/Experimental/ImageWidgets/Source/ImageWidgets/ImageWidgets.Build.cs b/Engine/Plugins/Experimental/ImageWidgets/Source/ImageWidgets/ImageWidgets.Build.cs
--- a/Engine/Plugins/Experimental/ImageWidgets/Source/ImageWidgets/ImageWidgets.Build.cs
+++ b/Engine/Plugins/Experimental/ImageWidgets/Source/ImageWidgets/ImageWidgets.Build.cs
@@ -4,14 +4,14 @@ using UnrealBuildTool;
 
 public class ImageWidgets : ModuleRules
 {
-	// Set this flag to true to force building the ColorViewer sample. Otherwise it will only be built in Debug builds. 
+	// Set this flag to true to force building the ColorViewer sample. Otherwise, it will only be built in Debug builds. 
 	private const bool ForceBuildColorViewerSample = false;
 
 	// Todo Remove build flags as soon as all the respective prototype code is migrated to this module.
 	// These flags are currently used to migrate already existing prototype functionality into this module without having to remove/change the functionality of
 	// the prototype applications.
+	private const bool EnableCatalog = true;
 	private const bool EnableAbComparison = false;
-	private const bool EnableCatalog = false;
 
 	public ImageWidgets(ReadOnlyTargetRules Target) : base(Target)
 	{
@@ -50,7 +50,7 @@ public class ImageWidgets : ModuleRules
 		}
 		PublicDefinitions.Add(bBuildColorViewerSample ? "IMAGE_WIDGETS_BUILD_COLOR_VIEWER_SAMPLE=1" : "IMAGE_WIDGETS_BUILD_COLOR_VIEWER_SAMPLE=0");
 
-		PublicDefinitions.Add(EnableAbComparison ? "IMAGE_WIDGETS_WITH_AB_COMPARISON=1" : "IMAGE_WIDGETS_WITH_AB_COMPARISON=0");
 		PublicDefinitions.Add(EnableCatalog ? "IMAGE_WIDGETS_WITH_CATALOG=1" : "IMAGE_WIDGETS_WITH_CATALOG=0");
+		PublicDefinitions.Add(EnableAbComparison ? "IMAGE_WIDGETS_WITH_AB_COMPARISON=1" : "IMAGE_WIDGETS_WITH_AB_COMPARISON=0");
 	}
 }
\ No newline at end of file
diff --git a/Engine/Plugins/Experimental/ImageWidgets/Source/ImageWidgets/Private/ColorViewerSample/ColorViewer.cpp b/Engine/Plugins/Experimental/ImageWidgets/Source/ImageWidgets/Private/ColorViewerSample/ColorViewer.cpp
--- a/Engine/Plugins/Experimental/ImageWidgets/Source/ImageWidgets/Private/ColorViewerSample/ColorViewer.cpp
+++ b/Engine/Plugins/Experimental/ImageWidgets/Source/ImageWidgets/Private/ColorViewerSample/ColorViewer.cpp
@@ -13,40 +13,72 @@ namespace UE::ImageWidgets::Sample
 {
 	IImageViewer::FImageInfo FColorViewer::GetCurrentImageInfo() const
 	{
-		return {ImageGuid, ImageSize, 0, bColorIsValid};
+		if (ColorIsValid(SelectedColorIndex))
+		{
+			return {Colors[SelectedColorIndex].Guid, ImageSize, 0, true};
+		}
+		return {FGuid(), FIntPoint::ZeroValue, 0, false};
 	}
 
 	void FColorViewer::DrawCurrentImage(FViewport* Viewport, FCanvas* Canvas, const FDrawProperties& Properties)
 	{
-		// Get color value after tone mapping.
-		const FLinearColor ToneMappedColor = ToneMapping.GetToneMappedColor(Color);
+		if (ColorIsValid(SelectedColorIndex))
+		{
+			// Get color value after tone mapping.
+			const FLinearColor ToneMappedColor = ToneMapping.GetToneMappedColor(Colors[SelectedColorIndex].Color);
 
-		// Draw simple quad with current tone mapped color.
-		// In a less trivial use case, this would require rendering quads with textures and the like. 
-		FCanvasTileItem Tile(Properties.Placement.Offset, Properties.Placement.Size, ToneMappedColor);
-		Canvas->DrawItem(Tile);
+			// Draw simple quad with current tone mapped color.
+			// In a less trivial use case, this would require rendering quads with textures and the like. 
+			FCanvasTileItem Tile(Properties.Placement.Offset, Properties.Placement.Size, ToneMappedColor);
+			Canvas->DrawItem(Tile);
+		}
 	}
 
 	TOptional<TVariant<FColor, FLinearColor>> FColorViewer::GetCurrentImagePixelColor(FIntPoint PixelCoords, int32 MipLevel) const
 	{
-		if (bColorIsValid)
+		if (ColorIsValid(SelectedColorIndex))
 		{
 			// Returns the current color as float values.
 			// In a less trivial use case, the pixel coordinates and potentially the MIP level would be needed to look up the color value.
-			return TVariant<FColor, FLinearColor>(TInPlaceType<FLinearColor>(), Color);
+			return TVariant<FColor, FLinearColor>(TInPlaceType<FColor>(), Colors[SelectedColorIndex].Color);
 		}
 		return {};
 	}
 
-	void FColorViewer::RandomizeColor()
+	void FColorViewer::OnImageSelected(const FGuid& Guid)
 	{
-		auto Random = []
+		if (ColorIsValid(Guid.B) && Colors[Guid.B].Guid == Guid)
 		{
-			return FMath::RandRange(0.0f, 1.0f);
-		};
+			SelectedColorIndex = Guid.B;
+		}
+	}
 
-		Color = {Random(), Random(), Random()};
-		bColorIsValid = true;
+	const FColorViewer::FColorItem* FColorViewer::AddColor()
+	{
+		Colors.Add({FGuid(1, Colors.Num(), 0, 0), {}, FDateTime::Now()});
+
+		SelectedColorIndex = Colors.Num() - 1;
+
+		RandomizeColor();
+
+		return &Colors[SelectedColorIndex];
+	}
+
+	const FColorViewer::FColorItem* FColorViewer::RandomizeColor()
+	{
+		if (ColorIsValid(SelectedColorIndex))
+		{
+			auto Random = []
+			{
+				return static_cast<uint8>(FMath::RandRange(0, 255));
+			};
+
+			Colors[SelectedColorIndex].Color = {Random(), Random(), Random()};
+
+			return &Colors[SelectedColorIndex];
+		}
+
+		return nullptr;
 	}
 
 	FToneMapping::EMode FColorViewer::GetToneMapping() const
@@ -58,8 +90,13 @@ namespace UE::ImageWidgets::Sample
 	{
 		ToneMapping.Mode = Mode;
 	}
+
+	bool FColorViewer::ColorIsValid(int32 Index) const
+	{
+		return SelectedColorIndex != INDEX_NONE && SelectedColorIndex < Colors.Num();
+	}
 }
 
 #undef LOCTEXT_NAMESPACE
 
-#endif // IMAGE_WIDGETS_BUILD_COLOR_VIEWER_SAMPLE
\ No newline at end of file
+#endif // IMAGE_WIDGETS_BUILD_COLOR_VIEWER_SAMPLE
diff --git a/Engine/Plugins/Experimental/ImageWidgets/Source/ImageWidgets/Private/ColorViewerSample/ColorViewer.h b/Engine/Plugins/Experimental/ImageWidgets/Source/ImageWidgets/Private/ColorViewerSample/ColorViewer.h
--- a/Engine/Plugins/Experimental/ImageWidgets/Source/ImageWidgets/Private/ColorViewerSample/ColorViewer.h
+++ b/Engine/Plugins/Experimental/ImageWidgets/Source/ImageWidgets/Private/ColorViewerSample/ColorViewer.h
@@ -41,33 +41,54 @@ namespace UE::ImageWidgets::Sample
 	class FColorViewer final : public IImageViewer
 	{
 	public:
+		/**
+		 * Necessary data for a color item.
+		 */
+		struct FColorItem
+		{
+			/** Unique identifier for each item */
+			FGuid Guid;
+
+			/** the actual color value */
+			FColor Color;
+
+			/** Timestamp for when the item was created */
+			FDateTime DateTime;
+		};
+
 		// IImageViewer overrides - begin
 		virtual FImageInfo GetCurrentImageInfo() const override;
 		virtual void DrawCurrentImage(FViewport* Viewport, FCanvas* Canvas, const FDrawProperties& Properties) override;
 		virtual TOptional<TVariant<FColor, FLinearColor>> GetCurrentImagePixelColor(FIntPoint PixelCoords, int32 MipLevel) const override;
+		virtual void OnImageSelected(const FGuid& Guid) override;
 		// IImageViewer overrides - end
 
-		/** Sets a random color as the "image". This is a simple proxy for the image content changing and/or users choosing different images to display. */
-		void RandomizeColor();
+		/** Adds a color item. */
+		const FColorItem* AddColor();
+		
+		/** Sets a random color as the current "image" as a simple proxy for the image content changing and/or users choosing different images to display. */
+		const FColorItem* RandomizeColor();
 
 		/** Access to tone mapping data. This is effectively used by the viewport toolbar extensions as well as when drawing the image. */
 		FToneMapping::EMode GetToneMapping() const;
 		void SetToneMapping(FToneMapping::EMode Mode);
 
+		/** Hardcoded values for the image size for all color.
+		 *  In a more realistic application, this value would depend on the actual current image. */
+		inline static const FIntPoint ImageSize = 512;
+
 	private:
+		/** Checks if a given index is a valid image. */
+		bool ColorIsValid(int32 Index) const;
+
 		/** The tone mapping data. */
 		FToneMapping ToneMapping = FToneMapping::EMode::RGB;
 
 		/** The color for the currently displayed image. */
-		FLinearColor Color;
+		TArray<FColorItem> Colors;
 
-		/** This is set to true once the color is initialized by the user, and evaluated to tell the viewport if the image is valid. */
-		bool bColorIsValid = false;
-
-		/** Hardcoded values for the image size and the GUID of the single image this is currently able to display.
-		 *  In a more realistic application, these value would depend on the actual images that are available. */
-		inline static constexpr FGuid ImageGuid = FGuid(1, 0, 0, 0);
-		inline static const FIntPoint ImageSize = 512;
+		/** Index of the currently selected image. */
+		int32 SelectedColorIndex = INDEX_NONE;
 	};
 }
 
diff --git a/Engine/Plugins/Experimental/ImageWidgets/Source/ImageWidgets/Private/ColorViewerSample/ColorViewerCommands.cpp b/Engine/Plugins/Experimental/ImageWidgets/Source/ImageWidgets/Private/ColorViewerSample/ColorViewerCommands.cpp
--- a/Engine/Plugins/Experimental/ImageWidgets/Source/ImageWidgets/Private/ColorViewerSample/ColorViewerCommands.cpp
+++ b/Engine/Plugins/Experimental/ImageWidgets/Source/ImageWidgets/Private/ColorViewerSample/ColorViewerCommands.cpp
@@ -3,8 +3,8 @@
 #if IMAGE_WIDGETS_BUILD_COLOR_VIEWER_SAMPLE
 
 #include "ColorViewerCommands.h"
-
 #include "ColorViewerStyle.h"
+
 #include "Framework/Commands/InputChord.h"
 
 #define LOCTEXT_NAMESPACE "ColorViewerCommands"
@@ -18,6 +18,7 @@ namespace UE::ImageWidgets::Sample
 
 	void FColorViewerCommands::RegisterCommands()
 	{
+		UI_COMMAND(AddColor, "Add", "Add color to catalog", EUserInterfaceActionType::Button, FInputChord(EKeys::A));
 		UI_COMMAND(RandomizeColor, "Randomize", "Set random color", EUserInterfaceActionType::Button, FInputChord(EKeys::R));
 
 		UI_COMMAND(ToneMappingRGB, "RGB", "Show full color", EUserInterfaceActionType::ToggleButton, FInputChord());
diff --git a/Engine/Plugins/Experimental/ImageWidgets/Source/ImageWidgets/Private/ColorViewerSample/ColorViewerCommands.h b/Engine/Plugins/Experimental/ImageWidgets/Source/ImageWidgets/Private/ColorViewerSample/ColorViewerCommands.h
--- a/Engine/Plugins/Experimental/ImageWidgets/Source/ImageWidgets/Private/ColorViewerSample/ColorViewerCommands.h
+++ b/Engine/Plugins/Experimental/ImageWidgets/Source/ImageWidgets/Private/ColorViewerSample/ColorViewerCommands.h
@@ -19,9 +19,10 @@ namespace UE::ImageWidgets::Sample
 
 		virtual void RegisterCommands() override;
 
-		TSharedPtr<FUICommandInfo> RandomizeColor; // Button and hotkey for setting a random color
-		TSharedPtr<FUICommandInfo> ToneMappingRGB; // Toggle button for switching to RGB (no tone mapping)
-		TSharedPtr<FUICommandInfo> ToneMappingLum; // Toggle button for switching to luminance (grayscale tone mapping)
+		TSharedPtr<FUICommandInfo> AddColor; // Adding a new color entry
+		TSharedPtr<FUICommandInfo> RandomizeColor; // Setting current entry to a random color
+		TSharedPtr<FUICommandInfo> ToneMappingRGB; // Switching to RGB (no tone mapping)
+		TSharedPtr<FUICommandInfo> ToneMappingLum; // Switching to luminance (grayscale tone mapping)
 	};
 }
 
diff --git a/Engine/Plugins/Experimental/ImageWidgets/Source/ImageWidgets/Private/ColorViewerSample/ColorViewerWidget.cpp b/Engine/Plugins/Experimental/ImageWidgets/Source/ImageWidgets/Private/ColorViewerSample/ColorViewerWidget.cpp
--- a/Engine/Plugins/Experimental/ImageWidgets/Source/ImageWidgets/Private/ColorViewerSample/ColorViewerWidget.cpp
+++ b/Engine/Plugins/Experimental/ImageWidgets/Source/ImageWidgets/Private/ColorViewerSample/ColorViewerWidget.cpp
@@ -21,7 +21,7 @@ namespace UE::ImageWidgets::Sample
 		// Create toolbar extensions for a button to randomize the displayed color as well as the tone mapping controls.
 		const TSharedPtr<FExtender> ToolbarExtender = MakeShared<FExtender>();
 		ToolbarExtender->AddToolBarExtension("ToolbarCenter", EExtensionHook::Before, CommandList,
-		                                     FToolBarExtensionDelegate::CreateSP(this, &SColorViewerWidget::AddRandomColorButton));
+		                                     FToolBarExtensionDelegate::CreateSP(this, &SColorViewerWidget::AddColorButtons));
 		ToolbarExtender->AddToolBarExtension("ToolbarRight", EExtensionHook::After, CommandList,
 		                                     FToolBarExtensionDelegate::CreateSP(this, &SColorViewerWidget::AddToneMappingButtons));
 
@@ -30,16 +30,28 @@ namespace UE::ImageWidgets::Sample
 			.VAlign(VAlign_Fill)
 			.HAlign(HAlign_Fill)
 			[
-				SAssignNew(Viewport, ImageWidgets::SImageViewport, ColorViewer.ToSharedRef())
-					.ToolbarExtender(ToolbarExtender)
-					.DrawSettings(SImageViewport::FDrawSettings{
-							.ClearColor = FLinearColor::Black,
-							.bBorderEnabled = true,
-							.BorderThickness = 1.0,
-							.BorderColor = FVector3f(0.2f),
-							.bBackgroundColorEnabled = false,
-							.bBackgroundCheckerEnabled = false
-						})
+				SAssignNew(Splitter, SSplitter)
+					.PhysicalSplitterHandleSize(2.0f)
+					+ SSplitter::Slot()
+						.Value(0.0f)
+						[
+							SAssignNew(Catalog, ImageWidgets::SImageCatalog)
+								.OnItemSelected_Lambda([this](const FGuid& ImageGuid) { ColorViewer->OnImageSelected(ImageGuid); })
+						]
+					+ SSplitter::Slot()
+						.Value(1.0f)
+						[
+							SAssignNew(Viewport, ImageWidgets::SImageViewport, ColorViewer.ToSharedRef())
+								.ToolbarExtender(ToolbarExtender)
+								.DrawSettings(SImageViewport::FDrawSettings{
+									.ClearColor = FLinearColor::Black,
+									.bBorderEnabled = true,
+									.BorderThickness = 1.0,
+									.BorderColor = FVector3f(0.2f),
+									.bBackgroundColorEnabled = false,
+									.bBackgroundCheckerEnabled = false
+								})
+						]
 			];
 	}
 
@@ -54,10 +66,13 @@ namespace UE::ImageWidgets::Sample
 		return FReply::Unhandled();
 	}
 
-	void SColorViewerWidget::AddRandomColorButton(FToolBarBuilder& ToolbarBuilder) const
+	void SColorViewerWidget::AddColorButtons(FToolBarBuilder& ToolbarBuilder) const
 	{
-		const FSlateIcon RandomizeColorIcon(FAppStyle::GetAppStyleSetName(), "FontEditor.Update");
+		const FSlateIcon AddColorIcon(FAppStyle::GetAppStyleSetName(), "FontEditor.Button_Add");
+		ToolbarBuilder.AddToolBarButton(FColorViewerCommands::Get().AddColor, NAME_None, TAttribute<FText>(), TAttribute<FText>(),
+										TAttribute<FSlateIcon>(AddColorIcon));
 
+		const FSlateIcon RandomizeColorIcon(FAppStyle::GetAppStyleSetName(), "FontEditor.Update");
 		ToolbarBuilder.AddToolBarButton(FColorViewerCommands::Get().RandomizeColor, NAME_None, TAttribute<FText>(), TAttribute<FText>(),
 		                                TAttribute<FSlateIcon>(RandomizeColorIcon));
 	}
@@ -80,15 +95,61 @@ namespace UE::ImageWidgets::Sample
 		ToolbarBuilder.EndBlockGroup();
 	}
 
+	TTuple<FText, FText, FText> GetColorItemMetaData(const FColor Color, FDateTime DateTime)
+	{
+		const FString HexColor = FString::Printf(TEXT("#%02X%02X%02X"), Color.R, Color.G, Color.B);
+		const FText Name = FText::Format(LOCTEXT("ColorEntryLabel", "{0}"), FText::FromString(HexColor));
+
+		const FText Info = FText::Format(
+			LOCTEXT("ColorEntryInfoLabel", "{0}"), FText::AsTime(DateTime, EDateTimeStyle::Short, FText::GetInvariantTimeZone()));
+
+		const FText ToolTip = FText::Format(
+			LOCTEXT("ColorEntryToolTip", "R {0}, G {1}, B {2}"), {Color.R, Color.G, Color.B});
+
+		return {Name, Info, ToolTip};
+	}
+
+	void SColorViewerWidget::AddColor()
+	{
+		if (const FColorViewer::FColorItem* ColorItem = ColorViewer->AddColor())
+		{
+			auto [Name, Info, ToolTip] = GetColorItemMetaData(ColorItem->Color, ColorItem->DateTime);
+
+			Catalog->AddItem(MakeShared<FImageCatalogItemData>(ColorItem->Guid, FSlateColorBrush(ColorItem->Color), Name, Info, ToolTip));
+			Catalog->SelectItem(ColorItem->Guid);
+
+			if (bCatalogCollapsedOnInit && Catalog->NumTotalItems() > 1 && Splitter->SlotAt(0).GetSizeValue() <= 0.0f)
+			{
+				Splitter->SlotAt(0).SetSizeValue(0.2f);
+				bCatalogCollapsedOnInit = false;
+			}
+		}
+	}
+
+	void SColorViewerWidget::RandomizeColor()
+	{
+		if (const FColorViewer::FColorItem* ColorItem = ColorViewer->RandomizeColor())
+		{
+			const auto [Name, Info, ToolTip] = GetColorItemMetaData(ColorItem->Color, ColorItem->DateTime);
+
+			Catalog->UpdateItem({ColorItem->Guid, FSlateColorBrush(ColorItem->Color), Name, Info, ToolTip});
+		}
+	}
+
 	void SColorViewerWidget::BindCommands()
 	{
 		const FColorViewerCommands& Commands = FColorViewerCommands::Get();
 
 		CommandList = MakeShared<FUICommandList>();
 
+		CommandList->MapAction(
+			Commands.AddColor,
+			FExecuteAction::CreateSP(this, &SColorViewerWidget::AddColor)
+		);
+
 		CommandList->MapAction(
 			Commands.RandomizeColor,
-			FExecuteAction::CreateSP(ColorViewer.ToSharedRef(), &FColorViewer::RandomizeColor)
+			FExecuteAction::CreateSP(this, &SColorViewerWidget::RandomizeColor)
 		);
 
 		CommandList->MapAction(
diff --git a/Engine/Plugins/Experimental/ImageWidgets/Source/ImageWidgets/Private/ColorViewerSample/ColorViewerWidget.h b/Engine/Plugins/Experimental/ImageWidgets/Source/ImageWidgets/Private/ColorViewerSample/ColorViewerWidget.h
--- a/Engine/Plugins/Experimental/ImageWidgets/Source/ImageWidgets/Private/ColorViewerSample/ColorViewerWidget.h
+++ b/Engine/Plugins/Experimental/ImageWidgets/Source/ImageWidgets/Private/ColorViewerSample/ColorViewerWidget.h
@@ -6,6 +6,7 @@
 
 #include "ColorViewer.h"
 
+#include <SImageCatalog.h>
 #include <SImageViewport.h>
 
 namespace UE::ImageWidgets::Sample
@@ -29,21 +30,36 @@ namespace UE::ImageWidgets::Sample
 		virtual FReply OnKeyDown(const FGeometry& Geometry, const FKeyEvent& KeyEvent) override;
 
 	private:
-		/** Functions for adding tone mapping buttons as viewport toolbar extensions. */
-		void AddRandomColorButton(FToolBarBuilder& ToolbarBuilder) const;
+		/** Binds all the commands. */
+		void BindCommands();
+
+		/** Functions for adding buttons as viewport toolbar extensions. */
+		void AddColorButtons(FToolBarBuilder& ToolbarBuilder) const;
 		void AddToneMappingButtons(FToolBarBuilder& ToolbarBuilder) const;
 
-		/** Binds all of the commands. */
-		void BindCommands();
+		/** Add a new color entry to the catalog. */
+		void AddColor();
+
+		/** Choose a random color for the current entry. */
+		void RandomizeColor();
 
 		/** The image viewer implementation that contains the image data and renders the image. */
 		TSharedPtr<FColorViewer> ColorViewer;
 
-		/** The image viewport in which the image gets displayed in. */
+		/** Adjustable divider between catalog on the left and viewport on the right. */
+		TSharedPtr<SSplitter> Splitter;
+
+		/** The image catalog that holds all currently available images. */
+		TSharedPtr<SImageCatalog> Catalog;
+		
+		/** The image viewport in which the current image gets displayed in. */
 		TSharedPtr<SImageViewport> Viewport;
 
 		/** The commands used by this sample widget. */
 		TSharedPtr<FUICommandList> CommandList;
+
+		/** Indicates that the catalog is still collapsed. This gets set to false as soon as a catalog entry is added, or it is manually expanded. */
+		bool bCatalogCollapsedOnInit = true;
 	};
 }
 
diff --git a/Engine/Plugins/Experimental/ImageWidgets/Source/ImageWidgets/Private/SImageCatalog.cpp b/Engine/Plugins/Experimental/ImageWidgets/Source/ImageWidgets/Private/SImageCatalog.cpp
new file mode 100644
--- /dev/null
+++ b/Engine/Plugins/Experimental/ImageWidgets/Source/ImageWidgets/Private/SImageCatalog.cpp
@@ -0,0 +1,222 @@
+﻿// Copyright Epic Games, Inc. All Rights Reserved.
+
+#include "SImageCatalog.h"
+
+#include "ImageWidgetsLogCategory.h"
+#include "SImageCatalogItem.h"
+#include "Brushes/SlateColorBrush.h"
+#include "Styling/StyleColors.h"
+#include "Widgets/Layout/SSeparator.h"
+
+#define LOCTEXT_NAMESPACE "SImageViewerCatalog"
+
+namespace UE::ImageWidgets
+{
+FImageCatalogItemData::FImageCatalogItemData(FGuid Guid, const FSlateBrush& Brush, const FText& Name, const FText& Info, const FText& ToolTip)
+	: Guid(Guid), Thumbnail(Brush), Name(Name), Info(Info), ToolTip(ToolTip)
+{
+}
+
+void SImageCatalog::Construct(const FArguments& InArgs)
+{
+	OnItemSelected = InArgs._OnItemSelected;
+
+	Items.Reserve(1000);
+
+	const auto GenerateItemRow = [](const TSharedPtr<FImageCatalogItemData>& ItemData, const TSharedRef<STableViewBase>& OwnerTable)
+	{
+		static const FTableRowStyle TableRowStyle = []
+		{
+			FTableRowStyle Style = FAppStyle::Get().GetWidgetStyle<FTableRowStyle>("TableView.Row");
+			Style.SetOddRowBackgroundBrush(FSlateColorBrush(FStyleColors::Background));
+			Style.SetOddRowBackgroundHoveredBrush(FSlateColorBrush(FStyleColors::SelectHover));
+			Style.SetEvenRowBackgroundBrush(FSlateColorBrush(FStyleColors::Recessed));
+			Style.SetEvenRowBackgroundHoveredBrush(FSlateColorBrush(FStyleColors::SelectHover));
+			return Style;
+		}();
+
+		TSharedPtr<SImageCatalogItem> ItemWidget;
+		SAssignNew(ItemWidget, SImageCatalogItem, ItemData);
+
+		return SNew(STableRow<TSharedPtr<FImageCatalogItemData>>, OwnerTable)
+			.Style(&TableRowStyle)
+			.ShowSelection(true)
+			[
+				ItemWidget.ToSharedRef()
+			];
+	};
+
+	const auto HasPinnedItems = [this]
+	{
+		return PinnedItems.IsEmpty() ? EVisibility::Collapsed : EVisibility::Visible;
+	};
+
+	const auto PinnedItemSelectionChanged = [this](const TSharedPtr<FImageCatalogItemData>& Item, ESelectInfo::Type SelectInfo)
+	{
+		// Note that Item might be a nullptr since this callback is also executed when clearing the selection of the list.
+		if (Item.IsValid() && OnItemSelected.IsBound())
+		{
+			ItemsListView->ClearSelection();
+			OnItemSelected.Execute(Item->Guid);
+		}
+	};
+
+	const auto ItemSelectionChanged = [this](const TSharedPtr<FImageCatalogItemData>& Item, ESelectInfo::Type SelectInfo)
+	{
+		// Note that Item might be a nullptr since this callback is also executed when clearing the selection of the list.
+		if (Item.IsValid() && OnItemSelected.IsBound())
+		{
+			PinnedItemsListView->ClearSelection();
+			OnItemSelected.Execute(Item->Guid);
+		}
+	};
+
+	SAssignNew(PinnedItemsListView, SListView<TSharedPtr<FImageCatalogItemData>>)
+					.ListItemsSource(&PinnedItems)
+					.OnGenerateRow_Lambda(GenerateItemRow)
+					.OnSelectionChanged_Lambda(PinnedItemSelectionChanged)
+					.ClearSelectionOnClick(false)
+					.Visibility_Lambda(HasPinnedItems);
+
+	SAssignNew(ItemsListView, SListView<TSharedPtr<FImageCatalogItemData>>)
+					.ListItemsSource(&Items)
+					.OnGenerateRow_Lambda(GenerateItemRow)
+					.OnSelectionChanged_Lambda(ItemSelectionChanged)
+					.ClearSelectionOnClick(false)
+					.ScrollbarVisibility(EVisibility::Visible);
+	
+	ChildSlot
+	[
+		SNew(SVerticalBox)
+		+ SVerticalBox::Slot()
+			.AutoHeight()
+			.Padding(2.0f, 4.0f, 2.0f, 4.0f)
+			[
+				SNew(STextBlock)
+					.Text(InArgs._PinnedItemsHeading)
+					.TextStyle(FAppStyle::Get(), "DetailsView.CategoryTextStyle")
+					.Visibility(InArgs._PinnedItemsHeading.IsEmpty() ? EVisibility::Collapsed : EVisibility::Visible)
+			]
+		+ SVerticalBox::Slot()
+			.AutoHeight()
+			[
+				PinnedItemsListView.ToSharedRef()
+			]
+		+ SVerticalBox::Slot()
+			.AutoHeight()
+			[
+				SNew(SSeparator)
+					.Thickness(6.0f)
+					.Visibility_Lambda(HasPinnedItems)
+			]
+		+ SVerticalBox::Slot()
+			.AutoHeight()
+			.Padding(2.0f, 4.0f, 2.0f, 4.0f)
+			[
+				SNew(STextBlock)
+					.Text(InArgs._ItemsHeading)
+					.TextStyle(FAppStyle::Get(), "DetailsView.CategoryTextStyle")
+					.Visibility(InArgs._ItemsHeading.IsEmpty() ? EVisibility::Collapsed : EVisibility::Visible)
+			]
+		+ SVerticalBox::Slot()
+			[
+				ItemsListView.ToSharedRef()
+			]
+	];
+}
+
+void SImageCatalog::AddItem(const TSharedPtr<FImageCatalogItemData>& Item)
+{
+	Items.Add(Item);
+	GuidToItemMapping.Add(Item->Guid, TPair<const TSharedPtr<FImageCatalogItemData>, bool>(Item, false));
+	ItemsListView->ReGenerateItems(GetCachedGeometry());
+}
+
+void SImageCatalog::AddPinnedItem(const TSharedPtr<FImageCatalogItemData>& Item)
+{
+	PinnedItems.Add(Item);
+	GuidToItemMapping.Add(Item->Guid, TPair<const TSharedPtr<FImageCatalogItemData>, bool>(Item, true));
+	PinnedItemsListView->ReGenerateItems(GetCachedGeometry());
+}
+
+int32 SImageCatalog::NumItems() const
+{
+	return Items.Num();
+}
+
+int32 SImageCatalog::NumPinnedItems() const
+{
+	return PinnedItems.Num();
+}
+
+int32 SImageCatalog::NumTotalItems() const
+{
+	return NumItems() + NumPinnedItems();
+}
+
+void SImageCatalog::SelectItem(const FGuid& Guid)
+{
+	if (const TPair<const TSharedPtr<FImageCatalogItemData>, bool> *const ExistingItem = FindItemData(Guid))
+	{
+		ItemsListView->ClearSelection();
+		PinnedItemsListView->ClearSelection();
+
+		SListView<TSharedPtr<FImageCatalogItemData>>* ListView = ExistingItem->Get<1>() ? PinnedItemsListView.Get() : ItemsListView.Get();
+
+		ListView->SetItemSelection(ExistingItem->Get<0>(), true);
+	}
+}
+
+const TPair<const TSharedPtr<FImageCatalogItemData>, bool>* SImageCatalog::FindItemData(const FGuid& Guid)
+{
+	if (const TPair<const TSharedPtr<FImageCatalogItemData>, bool>* ItemPtr = GuidToItemMapping.Find(Guid))
+	{
+		return ItemPtr;
+	}
+
+	UE_LOG(LogImageWidgets, Warning, TEXT("Cannot find catalog item for guid '%s'."), *Guid.ToString());
+	return nullptr;
+}
+
+void SImageCatalog::UpdateItem(const FImageCatalogItemData& Item)
+{
+	if (const TPair<const TSharedPtr<FImageCatalogItemData>, bool> *const ExistingItem = FindItemData(Item.Guid))
+	{
+		*ExistingItem->Get<0>() = Item;
+	}
+}
+
+void SImageCatalog::UpdateItemInfo(const FGuid& Guid, const FText& Info)
+{
+	if (const TPair<const TSharedPtr<FImageCatalogItemData>, bool> *const ExistingItem = FindItemData(Guid))
+	{
+		ExistingItem->Get<0>()->Info = Info;
+	}
+}
+
+void SImageCatalog::UpdateItemName(const FGuid& Guid, const FText& Name)
+{
+	if (const TPair<const TSharedPtr<FImageCatalogItemData>, bool> *const ExistingItem = FindItemData(Guid))
+	{
+		ExistingItem->Get<0>()->Name = Name;
+	}
+}
+
+void SImageCatalog::UpdateItemThumbnail(const FGuid& Guid, const FSlateBrush& Thumbnail)
+{
+	if (const TPair<const TSharedPtr<FImageCatalogItemData>, bool> *const ExistingItem = FindItemData(Guid))
+	{
+		ExistingItem->Get<0>()->Thumbnail = Thumbnail;
+	}
+}
+
+void SImageCatalog::UpdateItemToolTip(const FGuid& Guid, const FText& ToolTip)
+{
+	if (const TPair<const TSharedPtr<FImageCatalogItemData>, bool> *const ExistingItem = FindItemData(Guid))
+	{
+		ExistingItem->Get<0>()->ToolTip = ToolTip;
+	}
+}
+}
+
+#undef LOCTEXT_NAMESPACE
diff --git a/Engine/Plugins/Experimental/ImageWidgets/Source/ImageWidgets/Private/SImageCatalogItem.cpp b/Engine/Plugins/Experimental/ImageWidgets/Source/ImageWidgets/Private/SImageCatalogItem.cpp
new file mode 100644
--- /dev/null
+++ b/Engine/Plugins/Experimental/ImageWidgets/Source/ImageWidgets/Private/SImageCatalogItem.cpp
@@ -0,0 +1,79 @@
+﻿// Copyright Epic Games, Inc. All Rights Reserved.
+
+#include "SImageCatalogItem.h"
+#include "SImageCatalog.h"
+
+#define LOCTEXT_NAMESPACE "SImageViewerCatalogItem"
+
+namespace UE::ImageWidgets
+{
+void SImageCatalogItem::Construct(const FArguments& InArgs, const TSharedPtr<FImageCatalogItemData>& InItemData)
+{
+	ItemData = InItemData;
+
+	ChildSlot
+	[
+		SNew(SHorizontalBox)
+		+ SHorizontalBox::Slot()
+			.Padding(FMargin(2.0f, 2.0f))
+			.AutoWidth()
+		[
+			SNew(SImage)
+				.DesiredSizeOverride(FVector2D(32.0f, 32.0f))
+				.Image(this, &SImageCatalogItem::GetItemThumbnail)
+				.ToolTipText(this, &SImageCatalogItem::GetItemToolTip)
+		]
+		+ SHorizontalBox::Slot()
+			.FillWidth(1.0f)
+			.VAlign(VAlign_Center)
+			.HAlign(HAlign_Fill)
+		[
+			SNew(SVerticalBox)
+			+ SVerticalBox::Slot()
+				.AutoHeight()
+				.HAlign(HAlign_Fill)
+				.Padding(FMargin(1.0f, 1.0f))
+			[
+				SNew(STextBlock)
+					.Text(this, &SImageCatalogItem::GetItemName)
+			]
+			+ SVerticalBox::Slot()
+				.AutoHeight()
+				.HAlign(HAlign_Fill)
+				.Padding(FMargin(1.0f, 1.0f))
+			[
+			SNew(SHorizontalBox)
+				+ SHorizontalBox::Slot()
+					.FillWidth(1.0f)
+				[
+					SNew(STextBlock)
+						.Text(this, &SImageCatalogItem::GetItemInfo)
+						.TextStyle(&FAppStyle::Get(), "SmallText.Subdued")
+				]
+			]
+		]
+	];
+}
+
+FText SImageCatalogItem::GetItemInfo() const
+{
+	return ItemData->Info;
+}
+
+FText SImageCatalogItem::GetItemName() const
+{
+	return ItemData->Name;
+}
+
+const FSlateBrush* SImageCatalogItem::GetItemThumbnail() const
+{
+	return &ItemData->Thumbnail;
+}
+
+FText SImageCatalogItem::GetItemToolTip() const
+{
+	return ItemData->ToolTip;
+}
+}
+
+#undef LOCTEXT_NAMESPACE
diff --git a/Engine/Plugins/Experimental/ImageWidgets/Source/ImageWidgets/Private/SImageCatalogItem.h b/Engine/Plugins/Experimental/ImageWidgets/Source/ImageWidgets/Private/SImageCatalogItem.h
new file mode 100644
--- /dev/null
+++ b/Engine/Plugins/Experimental/ImageWidgets/Source/ImageWidgets/Private/SImageCatalogItem.h
@@ -0,0 +1,31 @@
+﻿// Copyright Epic Games, Inc. All Rights Reserved.
+
+#pragma once
+
+#include "CoreMinimal.h"
+#include "SImageCatalog.h"
+
+namespace UE::ImageWidgets
+{
+	struct FImageCatalogItemData;
+
+	/**
+	 * Widget for a single item row in the catalog.
+	 */
+	class SImageCatalogItem : public SCompoundWidget
+	{
+	public:
+		SLATE_BEGIN_ARGS(SImageCatalogItem) {}
+		SLATE_END_ARGS()
+
+		void Construct(const FArguments& InArgs, const TSharedPtr<FImageCatalogItemData>& InItemData);
+
+	private:
+		FText GetItemInfo() const;
+		FText GetItemName() const;
+		const FSlateBrush* GetItemThumbnail() const;
+		FText GetItemToolTip() const;
+
+		TSharedPtr<FImageCatalogItemData> ItemData;
+	};
+}
diff --git a/Engine/Plugins/Experimental/ImageWidgets/Source/ImageWidgets/Public/IImageViewer.h b/Engine/Plugins/Experimental/ImageWidgets/Source/ImageWidgets/Public/IImageViewer.h
--- a/Engine/Plugins/Experimental/ImageWidgets/Source/ImageWidgets/Public/IImageViewer.h
+++ b/Engine/Plugins/Experimental/ImageWidgets/Source/ImageWidgets/Public/IImageViewer.h
@@ -73,7 +73,7 @@ namespace UE::ImageWidgets
 		};
 
 		/**
-		 * Provide any necessary metadata for the image widgets about the image that is currently supposed to be displayed. This data is generic in the sense
+		 * Provides any necessary metadata for the image widgets about the image that is currently supposed to be displayed. This data is generic in the sense
 		 * that the image widgets don't need to know any of the image structure, its content or how to draw it. Instead, the image drawing is done directly
 		 * via @ref DrawCurrentImage.
 		 * @return information about the image to be displayed
@@ -89,13 +89,22 @@ namespace UE::ImageWidgets
 		virtual void DrawCurrentImage(FViewport* Viewport, FCanvas* Canvas, const FDrawProperties& Properties) = 0;
 		
 		/**
-		 * Provide information about a given pixel.
+		 * Provides information about a given pixel.
 		 * @param PixelCoords XY coordinates for the pixel
 		 * @param MipLevel MIP level that is currently displayed; this can be ignored for images not supporting MIPs 
 		 * @return Either a color value in byte or float format, i.e. FColor or FLinearColor, or no value if there is no valid pixel at the provided coordinates 
 		 */
 		virtual TOptional<TVariant<FColor, FLinearColor>> GetCurrentImagePixelColor(FIntPoint PixelCoords, int32 MipLevel) const = 0;
 
+#if IMAGE_WIDGETS_WITH_CATALOG
+		/**
+		 * Notifies about the image with the given GUID being selected.
+		 * This can be implemented as an empty function if the image viewer implementation does not support switching between different images.
+		 * @param Guid Unique identifier of the selected image
+		 */
+		virtual void OnImageSelected(const FGuid& Guid) = 0;
+#endif
+
 	protected:
 		~IImageViewer() = default;
 	};
diff --git a/Engine/Plugins/Experimental/ImageWidgets/Source/ImageWidgets/Public/SImageCatalog.h b/Engine/Plugins/Experimental/ImageWidgets/Source/ImageWidgets/Public/SImageCatalog.h
new file mode 100644
--- /dev/null
+++ b/Engine/Plugins/Experimental/ImageWidgets/Source/ImageWidgets/Public/SImageCatalog.h
@@ -0,0 +1,166 @@
+﻿// Copyright Epic Games, Inc. All Rights Reserved.
+
+#pragma once
+
+#include "CoreMinimal.h"
+
+#include <Widgets/SCompoundWidget.h>
+#include <Widgets/Views/SListView.h>
+
+namespace UE::ImageWidgets
+{
+	class SImageCatalogItem;
+
+	/**
+	 * Contains all data for a catalog item.
+	 */
+	struct IMAGEWIDGETS_API FImageCatalogItemData
+	{
+		FImageCatalogItemData(FGuid Guid, const FSlateBrush& Brush, const FText& Name, const FText& Info, const FText& ToolTip);
+
+		/** Unique identifier for the catalog item */
+		FGuid Guid;
+
+		/** Brush used for displaying the item's thumbnail */
+		FSlateBrush Thumbnail;
+
+		/** Name of the item */
+		FText Name;
+
+		/** Auxiliary information for the item */
+		FText Info;
+
+		/** Tooltip that is shown when hovering over any part of the item's widget in the catalog */
+		FText ToolTip;
+	};
+
+	/**
+	 * Generic catalog widget for listing and selecting 2D image-like content.
+	 * Each catalog item is represented by its own widget based on its @ref FImageCatalogItemData.
+	 * Entries in the catalog can be split into regular and pinned items, with pinned items being displayed at the top.
+	 */
+	class SImageCatalog : public SCompoundWidget
+	{
+	public:
+		/**
+		 * Delegate that gets called when an item is selected in the catalog.
+		 * The given @ref FGuid identifies the item that was selected.
+		 */
+		DECLARE_DELEGATE_OneParam(FOnItemSelected, const FGuid&)
+
+		SLATE_BEGIN_ARGS(SImageCatalog)
+			{
+			}
+
+			/** Header text for regular items */
+			SLATE_ARGUMENT(FText, ItemsHeading)
+
+			/** Header text for pinned items */
+			SLATE_ARGUMENT(FText, PinnedItemsHeading)
+
+			/** Delegate that gets called when an item is selected in the catalog. */
+			SLATE_EVENT(FOnItemSelected, OnItemSelected)
+		SLATE_END_ARGS()
+
+		/**
+		 * Function used by Slate to construct the image catalog widget with the given arguments.
+		 * @param InArgs Slate arguments defined above
+		 */
+		IMAGEWIDGETS_API void Construct(const FArguments& InArgs);
+
+		/**
+		 * Adds a regular item to the catalog.
+		 * @param Item Data for the item that is being added. 
+		 */
+		IMAGEWIDGETS_API void AddItem(const TSharedPtr<FImageCatalogItemData>& Item);
+
+		/**
+		 * Adds a pinned item to the catalog. Pinned items appear in a seperate list above regular items.
+		 * @param Item Data for the item that is being added. 
+		 */
+		IMAGEWIDGETS_API void AddPinnedItem(const TSharedPtr<FImageCatalogItemData>& Item);
+
+		/**
+		 * Returns the number of regular items in the catalog, i.e. items that are not pinned.
+		 */
+		IMAGEWIDGETS_API int32 NumItems() const;
+
+		/**
+		 * Returns the number of pinned items in the catalog.
+		 */
+		IMAGEWIDGETS_API int32 NumPinnedItems() const;
+
+		/**
+		 * Returns the total number of items in the catalog, i.e. both regular items and pinned items.
+		 */
+		IMAGEWIDGETS_API int32 NumTotalItems() const;
+
+		/**
+		 * Select an existing regular or pinned item.
+		 * Nothing happens if no item with the given unique identifier exists in the catalog.
+		 * @param Guid The unique identifier of the existing item.
+		 */
+		IMAGEWIDGETS_API void SelectItem(const FGuid& Guid);
+
+		/**
+		 * Update an existing regular or pinned item's data. The item data should contain the item's unique identifier.
+		 * Nothing happens if no item with the given unique identifier exists in the catalog.
+		 * @param Item Data for the item that is being updated, including the unique identifier of the existing item.
+		 */
+		IMAGEWIDGETS_API void UpdateItem(const FImageCatalogItemData& Item);
+
+		/**
+		 * Update the info text of an existing regular or pinned item.
+		 * Nothing happens if no item with the given unique identifier exists in the catalog.
+		 * @param Guid The unique identifier of the existing item.
+		 * @param Info Text for the info label in the item's widget.
+		 */
+		IMAGEWIDGETS_API void UpdateItemInfo(const FGuid& Guid, const FText& Info);
+
+		/**
+		 * Update the info text of an existing regular or pinned item.
+		 * Nothing happens if no item with the given unique identifier exists in the catalog.
+		 * @param Guid The unique identifier of the existing item.
+		 * @param Name Text for the name label in the item's widget.
+		 */
+		IMAGEWIDGETS_API void UpdateItemName(const FGuid& Guid, const FText& Name);
+
+		/**
+		 * Update the thumbnail of an existing regular or pinned item.
+		 * Nothing happens if no item with the given unique identifier exists in the catalog.
+		 * @param Guid The unique identifier of the existing item.
+		 * @param Thumbnail Brush used for the thumbnail in the item's widget.
+		 */
+		IMAGEWIDGETS_API void UpdateItemThumbnail(const FGuid& Guid, const FSlateBrush& Thumbnail);
+
+		/**
+		 * Update the tooltip text of an existing regular or pinned item.
+		 * Nothing happens if no item with the given unique identifier exists in the catalog.
+		 * @param Guid The unique identifier of the existing item.
+		 * @param ToolTip Text for the tooltip label in the item's widget.
+		 */
+		IMAGEWIDGETS_API void UpdateItemToolTip(const FGuid& Guid, const FText& ToolTip);
+
+	private:
+		/** Find an existing regular or pinned item's data. The second value is true if the item is pinned. */
+		const TPair<const TSharedPtr<FImageCatalogItemData>, bool>* FindItemData(const FGuid& Guid);
+
+		/** Widget for listing all regular items. */
+		TSharedPtr<SListView<TSharedPtr<FImageCatalogItemData>>> ItemsListView;
+
+		/** Widget for listing all pinned items. */
+		TSharedPtr<SListView<TSharedPtr<FImageCatalogItemData>>> PinnedItemsListView;
+
+		/** Regular item data. */
+		TArray<TSharedPtr<FImageCatalogItemData>> Items;
+
+		/** Regular item data. */
+		TArray<TSharedPtr<FImageCatalogItemData>> PinnedItems;
+
+		/** Mapping from the items' unique identifier to their respective data and if it is a pinned item. */
+		TMap<FGuid, TPair<const TSharedPtr<FImageCatalogItemData>, bool>> GuidToItemMapping;
+
+		/** Delegate that gets called when an item is selected. */
+		FOnItemSelected OnItemSelected;
+	};
+}
diff --git a/Engine/Plugins/Experimental/ImageWidgets/Source/ImageWidgets/Public/SImageViewport.h b/Engine/Plugins/Experimental/ImageWidgets/Source/ImageWidgets/Public/SImageViewport.h
--- a/Engine/Plugins/Experimental/ImageWidgets/Source/ImageWidgets/Public/SImageViewport.h
+++ b/Engine/Plugins/Experimental/ImageWidgets/Source/ImageWidgets/Public/SImageViewport.h
@@ -14,7 +14,7 @@ namespace UE::ImageWidgets
 	class FStatusBarExtension;
 	class IImageViewer;
 	class SImageViewportToolbar;
-	
+
 	/**
 	 * Generic viewport for displaying and interacting with 2D image-like content.
 	 * The drawing of the images is deferred to an @ref IImageViewer implementation that needs to be provided upon construction. This viewport only uses the 
@@ -101,34 +101,41 @@ namespace UE::ImageWidgets
 		 */
 		struct FControllerSettings
 		{
-			enum class EDefaultZoomMode { Fit, Fill };
+			enum class EDefaultZoomMode
+			{
+				Fit, // Make the image fit within the viewport, but do not make it larger than the original size.
+				Fill // Make the image fit within the viewport, and if it is smaller than the viewport, zoom in to fill the viewport.
+			};
 
+			/** Zoom mode that gets set on viewport construction and whenever the viewport controller is reset via @ref ResetController(). */
 			EDefaultZoomMode DefaultZoomMode = EDefaultZoomMode::Fit;
 		};
 
 		DECLARE_DELEGATE(FOnLeftMouseButtonPressed);
 		DECLARE_DELEGATE(FOnLeftMouseButtonReleased);
 
-		SLATE_BEGIN_ARGS(SImageViewport){}
+		SLATE_BEGIN_ARGS(SImageViewport)
+			{
+			}
+
+			/** Extensions for the viewport toolbar; valid hooks are "ToolbarLeft", "ToolbarCenter", "ToolbarRight" */
+			SLATE_ARGUMENT(TSharedPtr<FExtender>, ToolbarExtender)
 
-		/** Extensions for the viewport toolbar; valid hooks are "ToolbarLeft", "ToolbarCenter", "ToolbarRight" */
-		SLATE_ARGUMENT(TSharedPtr<FExtender>, ToolbarExtender)
+			/** Extensions for the viewport status bar; valid hooks are "StatusBarLeft", "StatusBarCenter", "StatusBarRight" */
+			SLATE_ARGUMENT(TSharedPtr<FStatusBarExtender>, StatusBarExtender)
 
-		/** Extensions for the viewport status bar; valid hooks are "StatusBarLeft", "StatusBarCenter", "StatusBarRight" */
-		SLATE_ARGUMENT(TSharedPtr<FStatusBarExtender>, StatusBarExtender)
+			/** Settings for drawing viewport contents other than the actual image */
+			SLATE_ATTRIBUTE(FDrawSettings, DrawSettings)
 
-		/** Settings for drawing viewport contents other than the actual image */
-		SLATE_ATTRIBUTE(FDrawSettings, DrawSettings)
+			/** Settings for controlling the viewport */
+			SLATE_ARGUMENT(FControllerSettings, ControllerSettings)
 
-		/** Settings for controlling the viewport */
-		SLATE_ARGUMENT(FControllerSettings, ControllerSettings)
+			/** Left mouse button pressed event bubbled up from the viewport */
+			SLATE_EVENT(FOnLeftMouseButtonPressed, OnLeftMouseButtonPressed)
 
-		/** Left mouse button pressed event bubbled up from the viewport */
-		SLATE_EVENT(FOnLeftMouseButtonPressed, OnLeftMouseButtonPressed)
+			/** Left mouse button released event bubbled up from the viewport */
+			SLATE_EVENT(FOnLeftMouseButtonReleased, OnLeftMouseButtonReleased)
 
-		/** Left mouse button released event bubbled up from the viewport */
-		SLATE_EVENT(FOnLeftMouseButtonReleased, OnLeftMouseButtonReleased)
-		
 		SLATE_END_ARGS()
 
 		IMAGEWIDGETS_API SImageViewport();
