Added functionality to expose the MouseCaptureMode from the ImageWidget.

diff --git a/Engine/Plugins/Experimental/ImageWidgets/Source/ImageWidgets/Private/ImageViewportClient.cpp b/Engine/Plugins/Experimental/ImageWidgets/Source/ImageWidgets/Private/ImageViewportClient.cpp
--- a/Engine/Plugins/Experimental/ImageWidgets/Source/ImageWidgets/Private/ImageViewportClient.cpp
+++ b/Engine/Plugins/Experimental/ImageWidgets/Source/ImageWidgets/Private/ImageViewportClient.cpp
@@ -99,7 +99,9 @@ void DestroyCheckerTexture(TStrongObjectPtr<UTexture2D>& CheckerTexture)
 
 FImageViewportClient::FImageViewportClient(const TWeakPtr<SEditorViewport>& InEditorViewport, FGetImageSize&& InGetImageSize, FDrawImage&& InDrawImage,
                                            FGetDrawSettings&& InGetDrawSettings, FGetDPIScaleFactor&& InGetDPIScaleFactor,
-                                           FOnLeftMouseButtonPressed&& InOnLeftMouseButtonPressed, FOnLeftMouseButtonReleased&& InOnLeftMouseButtonReleased, SImageViewport::FControllerSettings::EDefaultZoomMode DefaultZoomMode)
+                                           FOnLeftMouseButtonPressed&& InOnLeftMouseButtonPressed, FOnLeftMouseButtonReleased&& InOnLeftMouseButtonReleased,
+                                           SImageViewport::FControllerSettings::EDefaultZoomMode DefaultZoomMode,
+                                           EMouseCaptureMode InMouseCaptureMode /*= EMouseCaptureMode::CapturePermanently*/)
 	: FEditorViewportClient(nullptr, nullptr, InEditorViewport)
 	, GetImageSize(MoveTemp(InGetImageSize))
 	, DrawImage(MoveTemp(InDrawImage))
@@ -108,6 +110,7 @@ FImageViewportClient::FImageViewportClient(const TWeakPtr<SEditorViewport>& InEd
 	, OnLeftMouseButtonPressed(MoveTemp(InOnLeftMouseButtonPressed))
 	, OnLeftMouseButtonReleased(MoveTemp(InOnLeftMouseButtonReleased))
 	, Controller(static_cast<FImageViewportController::EZoomMode>(DefaultZoomMode))
+	, MouseCaptureMode(InMouseCaptureMode)
 {
 	check(GetImageSize.IsBound());
 	check(DrawImage.IsBound());
@@ -403,6 +406,11 @@ FVector2d FImageViewportClient::GetViewportSizeWithDPIScaling() const
 
 	return FVector2d(ViewportSize) / DPIScaleFactor;
 }
+
+EMouseCaptureMode FImageViewportClient::GetMouseCaptureMode() const
+{
+	return MouseCaptureMode;
+}
 }
 
 #undef LOCTEXT_NAMESPACE
diff --git a/Engine/Plugins/Experimental/ImageWidgets/Source/ImageWidgets/Private/ImageViewportClient.h b/Engine/Plugins/Experimental/ImageWidgets/Source/ImageWidgets/Private/ImageViewportClient.h
--- a/Engine/Plugins/Experimental/ImageWidgets/Source/ImageWidgets/Private/ImageViewportClient.h
+++ b/Engine/Plugins/Experimental/ImageWidgets/Source/ImageWidgets/Private/ImageViewportClient.h
@@ -26,7 +26,8 @@ namespace UE::ImageWidgets
 	public:
 		FImageViewportClient(const TWeakPtr<SEditorViewport>& InViewport, FGetImageSize&& InGetImageSize, FDrawImage&& InDrawImage,
 		                     FGetDrawSettings&& InGetDrawSettings, FGetDPIScaleFactor&& InGetDPIScaleFactor,
-		                     FOnLeftMouseButtonPressed&& InOnLeftMouseButtonPressed, FOnLeftMouseButtonReleased&& InOnLeftMouseButtonReleased, SImageViewport::FControllerSettings::EDefaultZoomMode DefaultZoomMode);
+		                     FOnLeftMouseButtonPressed&& InOnLeftMouseButtonPressed, FOnLeftMouseButtonReleased&& InOnLeftMouseButtonReleased,
+		                     SImageViewport::FControllerSettings::EDefaultZoomMode DefaultZoomMode, EMouseCaptureMode InMouseCaptureMode = EMouseCaptureMode::CapturePermanently);
 		virtual ~FImageViewportClient() override;
 
 		virtual void Draw(FViewport* InViewport, FCanvas* Canvas) override;
@@ -36,6 +37,7 @@ namespace UE::ImageWidgets
 		virtual bool InputKey(const FInputKeyEventArgs& EventArgs) override;
 		virtual void TrackingStarted(const FInputEventState& InputState, bool bIsDraggingWidget, bool bNudge) override;
 		virtual void TrackingStopped() override;
+		virtual EMouseCaptureMode GetMouseCaptureMode() const override;
 		// SEditorViewport overrides - end
 
 		int32 GetMipLevel() const;
@@ -86,5 +88,7 @@ namespace UE::ImageWidgets
 
 		TStrongObjectPtr<UTexture2D> CheckerTexture;
 		FCheckerTextureSettings CachedCheckerTextureSettings;
+
+		EMouseCaptureMode MouseCaptureMode;
 	};
 }
diff --git a/Engine/Plugins/Experimental/ImageWidgets/Source/ImageWidgets/Private/SImageViewport.cpp b/Engine/Plugins/Experimental/ImageWidgets/Source/ImageWidgets/Private/SImageViewport.cpp
--- a/Engine/Plugins/Experimental/ImageWidgets/Source/ImageWidgets/Private/SImageViewport.cpp
+++ b/Engine/Plugins/Experimental/ImageWidgets/Source/ImageWidgets/Private/SImageViewport.cpp
@@ -198,7 +198,8 @@ namespace UE::ImageWidgets
 			FGetDPIScaleFactor::CreateLambda(GetDPIScaleFactor),
 			FOnLeftMouseButtonPressed::CreateLambda(HandleOnLeftMouseButtonPressed),
 			FOnLeftMouseButtonReleased::CreateLambda(HandleOnLeftMouseButtonReleased),
-			InArgs._ControllerSettings.DefaultZoomMode)
+			InArgs._ControllerSettings.DefaultZoomMode,
+			InArgs._MouseCaptureMode)
 			);
 
 		SEditorViewport::Construct(SEditorViewport::FArguments());
diff --git a/Engine/Plugins/Experimental/ImageWidgets/Source/ImageWidgets/Public/SImageViewport.h b/Engine/Plugins/Experimental/ImageWidgets/Source/ImageWidgets/Public/SImageViewport.h
--- a/Engine/Plugins/Experimental/ImageWidgets/Source/ImageWidgets/Public/SImageViewport.h
+++ b/Engine/Plugins/Experimental/ImageWidgets/Source/ImageWidgets/Public/SImageViewport.h
@@ -136,6 +136,9 @@ namespace UE::ImageWidgets
 			/** Left mouse button released event bubbled up from the viewport */
 			SLATE_EVENT(FOnLeftMouseButtonReleased, OnLeftMouseButtonReleased)
 
+			/** Settings for handling the mouse capture */
+			SLATE_ARGUMENT(EMouseCaptureMode, MouseCaptureMode)
+
 		SLATE_END_ARGS()
 
 		IMAGEWIDGETS_API SImageViewport();
