Assigned default value to  SImageViewport EMouseCapture slate argument.

diff --git a/Engine/Plugins/Experimental/ImageWidgets/Source/ImageWidgets/Public/SImageViewport.h b/Engine/Plugins/Experimental/ImageWidgets/Source/ImageWidgets/Public/SImageViewport.h
--- a/Engine/Plugins/Experimental/ImageWidgets/Source/ImageWidgets/Public/SImageViewport.h
+++ b/Engine/Plugins/Experimental/ImageWidgets/Source/ImageWidgets/Public/SImageViewport.h
@@ -114,7 +114,8 @@ namespace UE::ImageWidgets
 		DECLARE_DELEGATE(FOnLeftMouseButtonPressed);
 		DECLARE_DELEGATE(FOnLeftMouseButtonReleased);
 
-		SLATE_BEGIN_ARGS(SImageViewport)
+		SLATE_BEGIN_ARGS(SImageViewport) :
+			_MouseCaptureMode(EMouseCaptureMode::CapturePermanently)
 			{
 			}
 
