ImageWidgets: Bind +/- keys to zoom in SImageViewport

diff --git a/Engine/Plugins/Experimental/ImageWidgets/Source/ImageWidgets/Private/ImageViewportClient.cpp b/Engine/Plugins/Experimental/ImageWidgets/Source/ImageWidgets/Private/ImageViewportClient.cpp
--- a/Engine/Plugins/Experimental/ImageWidgets/Source/ImageWidgets/Private/ImageViewportClient.cpp
+++ b/Engine/Plugins/Experimental/ImageWidgets/Source/ImageWidgets/Private/ImageViewportClient.cpp
@@ -209,7 +209,7 @@ bool FImageViewportClient::InputKey(const FInputKeyEventArgs& EventArgs)
 {
 	if (EventArgs.Event == IE_Pressed)
 	{
-		if (EventArgs.Key == EKeys::MouseScrollUp)
+		if (EventArgs.Key == EKeys::MouseScrollUp || EventArgs.Key == EKeys::Add)
 		{
 			const auto [PixelCoordsValid, PixelCoords] = GetPixelCoordinatesUnderCursor();
 			if (PixelCoordsValid)
@@ -219,7 +219,7 @@ bool FImageViewportClient::InputKey(const FInputKeyEventArgs& EventArgs)
 			return true;
 		}
 
-		if (EventArgs.Key == EKeys::MouseScrollDown)
+		if (EventArgs.Key == EKeys::MouseScrollDown || EventArgs.Key == EKeys::Subtract)
 		{
 			const auto [PixelCoordsValid, PixelCoords] = GetPixelCoordinatesUnderCursor();
 			if (PixelCoordsValid)
diff --git a/Engine/Plugins/Experimental/ImageWidgets/Source/ImageWidgets/Private/ImageWidgetsCommands.cpp b/Engine/Plugins/Experimental/ImageWidgets/Source/ImageWidgets/Private/ImageWidgetsCommands.cpp
--- a/Engine/Plugins/Experimental/ImageWidgets/Source/ImageWidgets/Private/ImageWidgetsCommands.cpp
+++ b/Engine/Plugins/Experimental/ImageWidgets/Source/ImageWidgets/Private/ImageWidgetsCommands.cpp
@@ -16,8 +16,8 @@ namespace UE::ImageWidgets
 
 	void FImageWidgetsCommands::RegisterCommands()
 	{
-		UI_COMMAND(MipMinus, "Lower", "Show next lower mip", EUserInterfaceActionType::None, FInputChord(EKeys::Subtract));
-		UI_COMMAND(MipPlus, "Higher", "Show next higher mip", EUserInterfaceActionType::None, FInputChord(EKeys::Add));
+		UI_COMMAND(MipMinus, "Lower", "Show next lower mip", EUserInterfaceActionType::None, FInputChord());
+		UI_COMMAND(MipPlus, "Higher", "Show next higher mip", EUserInterfaceActionType::None, FInputChord());
 
 		UI_COMMAND(ToggleOverlay, "Toggle Overlay", "Toggles the visibility of the toolbar and the status bar.", EUserInterfaceActionType::ToggleButton,
 		           FInputChord(EKeys::O));
