ImageWidgets: Add A|B Comparison

diff --git a/Engine/Plugins/Experimental/ImageWidgets/Source/ImageWidgets/ImageWidgets.Build.cs b/Engine/Plugins/Experimental/ImageWidgets/Source/ImageWidgets/ImageWidgets.Build.cs
--- a/Engine/Plugins/Experimental/ImageWidgets/Source/ImageWidgets/ImageWidgets.Build.cs
+++ b/Engine/Plugins/Experimental/ImageWidgets/Source/ImageWidgets/ImageWidgets.Build.cs
@@ -8,10 +8,7 @@ public class ImageWidgets : ModuleRules
 	private const bool ForceBuildColorViewerSample = false;
 
 	// Todo Remove build flags as soon as all the respective prototype code is migrated to this module.
-	// These flags are currently used to migrate already existing prototype functionality into this module without having to remove/change the functionality of
-	// the prototype applications.
-	private const bool EnableCatalog = true;
-	private const bool EnableAbComparison = false;
+	private const bool EnableAbComparison = true;
 
 	public ImageWidgets(ReadOnlyTargetRules Target) : base(Target)
 	{
@@ -50,7 +47,6 @@ public class ImageWidgets : ModuleRules
 		}
 		PublicDefinitions.Add(bBuildColorViewerSample ? "IMAGE_WIDGETS_BUILD_COLOR_VIEWER_SAMPLE=1" : "IMAGE_WIDGETS_BUILD_COLOR_VIEWER_SAMPLE=0");
 
-		PublicDefinitions.Add(EnableCatalog ? "IMAGE_WIDGETS_WITH_CATALOG=1" : "IMAGE_WIDGETS_WITH_CATALOG=0");
 		PublicDefinitions.Add(EnableAbComparison ? "IMAGE_WIDGETS_WITH_AB_COMPARISON=1" : "IMAGE_WIDGETS_WITH_AB_COMPARISON=0");
 	}
 }
\ No newline at end of file
diff --git a/Engine/Plugins/Experimental/ImageWidgets/Source/ImageWidgets/Private/ColorViewerSample/ColorViewer.cpp b/Engine/Plugins/Experimental/ImageWidgets/Source/ImageWidgets/Private/ColorViewerSample/ColorViewer.cpp
--- a/Engine/Plugins/Experimental/ImageWidgets/Source/ImageWidgets/Private/ColorViewerSample/ColorViewer.cpp
+++ b/Engine/Plugins/Experimental/ImageWidgets/Source/ImageWidgets/Private/ColorViewerSample/ColorViewer.cpp
@@ -22,16 +22,19 @@ namespace UE::ImageWidgets::Sample
 
 	void FColorViewer::DrawCurrentImage(FViewport* Viewport, FCanvas* Canvas, const FDrawProperties& Properties)
 	{
-		if (ColorIsValid(SelectedColorIndex))
+#if IMAGE_WIDGETS_WITH_AB_COMPARISON
+		if (Properties.ABComparison.IsActive())
 		{
-			// Get color value after tone mapping.
-			const FLinearColor ToneMappedColor = ToneMapping.GetToneMappedColor(Colors[SelectedColorIndex].Color);
-
-			// Draw simple quad with current tone mapped color.
-			// In a less trivial use case, this would require rendering quads with textures and the like. 
-			FCanvasTileItem Tile(Properties.Placement.Offset, Properties.Placement.Size, ToneMappedColor);
-			Canvas->DrawItem(Tile);
+			DrawImage(Properties.ABComparison.GuidA.B, Canvas, Properties.Placement, {0.0, 0.0}, {Properties.ABComparison.Threshold, 1.0});
+			DrawImage(Properties.ABComparison.GuidB.B, Canvas, Properties.Placement, {Properties.ABComparison.Threshold, 0.0}, {1.0, 1.0});
 		}
+		else
+		{
+#endif
+			DrawImage(SelectedColorIndex, Canvas, Properties.Placement, {0.0, 0.0}, {1.0, 1.0});
+#if IMAGE_WIDGETS_WITH_AB_COMPARISON
+		}
+#endif
 	}
 
 	TOptional<TVariant<FColor, FLinearColor>> FColorViewer::GetCurrentImagePixelColor(FIntPoint PixelCoords, int32 MipLevel) const
@@ -53,6 +56,13 @@ namespace UE::ImageWidgets::Sample
 		}
 	}
 
+#if IMAGE_WIDGETS_WITH_AB_COMPARISON
+	bool FColorViewer::IsValidImage(const FGuid& Guid) const
+	{
+		return ColorIsValid(Guid.B) && Colors[Guid.B].Guid == Guid;
+	}
+#endif
+
 	const FColorViewer::FColorItem* FColorViewer::AddColor()
 	{
 		Colors.Add({FGuid(1, Colors.Num(), 0, 0), {}, FDateTime::Now()});
@@ -91,9 +101,32 @@ namespace UE::ImageWidgets::Sample
 		ToneMapping.Mode = Mode;
 	}
 
+	FLinearColor FColorViewer::GetDefaultToneMappedColor(const FColor& Color) const
+	{
+		return ToneMapping.GetToneMappedColor(Color);
+	}
+
 	bool FColorViewer::ColorIsValid(int32 Index) const
 	{
-		return SelectedColorIndex != INDEX_NONE && SelectedColorIndex < Colors.Num();
+		return 0 <= Index && Index < Colors.Num();
+	}
+
+	void FColorViewer::DrawImage(int32 Index, FCanvas* Canvas, const FDrawProperties::FPlacement& Placement, const FVector2d& UV0, const FVector2d& UV1) const
+	{
+		if (ColorIsValid(Index))
+		{
+			// Get color value after tone mapping.
+			const FLinearColor ToneMappedColor = ToneMapping.GetToneMappedColor(Colors[Index].Color);
+
+			// Adjust offset and size based on which part of the image to draw.
+			const FVector2d Offset = Placement.Offset + Placement.Size * UV0;
+			const FVector2d Size = Placement.Size * (UV1 - UV0);
+
+			// Draw simple quad with current tone mapped color.
+			// In a less trivial use case, this would require rendering quads with textures and the like. 
+			FCanvasTileItem Tile(Offset, Size, ToneMappedColor);
+			Canvas->DrawItem(Tile);
+		}
 	}
 }
 
diff --git a/Engine/Plugins/Experimental/ImageWidgets/Source/ImageWidgets/Private/ColorViewerSample/ColorViewer.h b/Engine/Plugins/Experimental/ImageWidgets/Source/ImageWidgets/Private/ColorViewerSample/ColorViewer.h
--- a/Engine/Plugins/Experimental/ImageWidgets/Source/ImageWidgets/Private/ColorViewerSample/ColorViewer.h
+++ b/Engine/Plugins/Experimental/ImageWidgets/Source/ImageWidgets/Private/ColorViewerSample/ColorViewer.h
@@ -49,7 +49,7 @@ namespace UE::ImageWidgets::Sample
 			/** Unique identifier for each item */
 			FGuid Guid;
 
-			/** the actual color value */
+			/** The actual color value */
 			FColor Color;
 
 			/** Timestamp for when the item was created */
@@ -61,6 +61,9 @@ namespace UE::ImageWidgets::Sample
 		virtual void DrawCurrentImage(FViewport* Viewport, FCanvas* Canvas, const FDrawProperties& Properties) override;
 		virtual TOptional<TVariant<FColor, FLinearColor>> GetCurrentImagePixelColor(FIntPoint PixelCoords, int32 MipLevel) const override;
 		virtual void OnImageSelected(const FGuid& Guid) override;
+#if IMAGE_WIDGETS_WITH_AB_COMPARISON
+		virtual bool IsValidImage(const FGuid& Guid) const override;
+#endif
 		// IImageViewer overrides - end
 
 		/** Adds a color item. */
@@ -73,6 +76,9 @@ namespace UE::ImageWidgets::Sample
 		FToneMapping::EMode GetToneMapping() const;
 		void SetToneMapping(FToneMapping::EMode Mode);
 
+		/** Apply default tone mapping to a given color. This is used to generate the catalog thumbnail. */
+		FLinearColor GetDefaultToneMappedColor(const FColor& Color) const;
+
 		/** Hardcoded values for the image size for all color.
 		 *  In a more realistic application, this value would depend on the actual current image. */
 		inline static const FIntPoint ImageSize = 512;
@@ -81,6 +87,9 @@ namespace UE::ImageWidgets::Sample
 		/** Checks if a given index is a valid image. */
 		bool ColorIsValid(int32 Index) const;
 
+		/** Draws the color image with the given index. The UVs determine if all or only a part of the image is drawn, i.e. for AB comparisons. */
+		void DrawImage(int32 Index, FCanvas* Canvas, const FDrawProperties::FPlacement& Placement, const FVector2d& UV0, const FVector2d& UV1) const;
+		
 		/** The tone mapping data. */
 		FToneMapping ToneMapping = FToneMapping::EMode::RGB;
 
diff --git a/Engine/Plugins/Experimental/ImageWidgets/Source/ImageWidgets/Private/ColorViewerSample/ColorViewerWidget.cpp b/Engine/Plugins/Experimental/ImageWidgets/Source/ImageWidgets/Private/ColorViewerSample/ColorViewerWidget.cpp
--- a/Engine/Plugins/Experimental/ImageWidgets/Source/ImageWidgets/Private/ColorViewerSample/ColorViewerWidget.cpp
+++ b/Engine/Plugins/Experimental/ImageWidgets/Source/ImageWidgets/Private/ColorViewerSample/ColorViewerWidget.cpp
@@ -51,6 +51,9 @@ namespace UE::ImageWidgets::Sample
 									.bBackgroundColorEnabled = false,
 									.bBackgroundCheckerEnabled = false
 								})
+#if IMAGE_WIDGETS_WITH_AB_COMPARISON
+								.bABComparisonEnabled(true)
+#endif
 						]
 			];
 	}
@@ -114,8 +117,9 @@ namespace UE::ImageWidgets::Sample
 		if (const FColorViewer::FColorItem* ColorItem = ColorViewer->AddColor())
 		{
 			auto [Name, Info, ToolTip] = GetColorItemMetaData(ColorItem->Color, ColorItem->DateTime);
+			const FLinearColor ToneMappedColor = ColorViewer->GetDefaultToneMappedColor(ColorItem->Color);
 
-			Catalog->AddItem(MakeShared<FImageCatalogItemData>(ColorItem->Guid, FSlateColorBrush(ColorItem->Color), Name, Info, ToolTip));
+			Catalog->AddItem(MakeShared<FImageCatalogItemData>(ColorItem->Guid, FSlateColorBrush(ToneMappedColor), Name, Info, ToolTip));
 			Catalog->SelectItem(ColorItem->Guid);
 
 			if (bCatalogCollapsedOnInit && Catalog->NumTotalItems() > 1 && Splitter->SlotAt(0).GetSizeValue() <= 0.0f)
diff --git a/Engine/Plugins/Experimental/ImageWidgets/Source/ImageWidgets/Private/ImageABComparison.h b/Engine/Plugins/Experimental/ImageWidgets/Source/ImageWidgets/Private/ImageABComparison.h
new file mode 100644
--- /dev/null
+++ b/Engine/Plugins/Experimental/ImageWidgets/Source/ImageWidgets/Private/ImageABComparison.h
@@ -0,0 +1,77 @@
+﻿// Copyright Epic Games, Inc. All Rights Reserved.
+
+#pragma once
+
+#include "Delegates/Delegate.h"
+#include "Misc/Guid.h"
+
+#if IMAGE_WIDGETS_WITH_AB_COMPARISON
+namespace UE::ImageWidgets
+{
+	/**
+	 * Provides data and logic for AB comparisons of images.
+	 */
+	class FImageABComparison
+	{
+	public:
+		DECLARE_DELEGATE_RetVal(FGuid, FGetCurrentImageGuid)
+		DECLARE_DELEGATE_RetVal_OneParam(bool, FImageIsValid, const FGuid&)
+
+		FImageABComparison(FImageIsValid&& ImageIsValid, FGetCurrentImageGuid&& GetCurrentImageGuid)
+			: ImageIsValid(ImageIsValid)
+			, GetCurrentImageGuid(GetCurrentImageGuid)
+		{
+			check(ImageIsValid.IsBound());
+			check(GetCurrentImageGuid.IsBound());
+		}
+
+		enum EAorB
+		{
+			A,
+			B
+		};
+
+		bool CanSetABComparison(EAorB AorB) const
+		{
+			const FGuid CurrentGuid = GetCurrentImageGuid.Execute();
+			const bool bGuidIsValidAndCurrentGuid = Guids[static_cast<int32>(AorB)].IsValid() && Guids[static_cast<int32>(AorB)] == CurrentGuid;
+			const bool bViewerHasValidImage = ImageIsValid.Execute(CurrentGuid);
+			const bool bGuidIsDisabledAndOtherGuidIsNotCurrentGuid = !Guids[static_cast<int32>(AorB)].IsValid() && Guids[!static_cast<bool>(AorB)] != CurrentGuid;
+
+			return bGuidIsValidAndCurrentGuid || (bViewerHasValidImage && bGuidIsDisabledAndOtherGuidIsNotCurrentGuid);
+		}
+
+		void SetABComparison(EAorB AorB, const FGuid& Guid)
+		{
+			Guids[static_cast<int32>(AorB)] = Guid;
+		}
+
+		bool ABComparisonIsSet(EAorB AorB) const
+		{
+			return Guids[static_cast<int32>(AorB)].IsValid();
+		}
+
+		bool IsActive() const
+		{
+			return Guids[0].IsValid() && Guids[1].IsValid();
+		}
+
+		const FGuid& GuidA() const
+		{
+			return Guids[0];
+		}
+
+		const FGuid& GuidB() const
+		{
+			return Guids[1];
+		}
+
+	private:
+		FGuid Guids[2];
+		FImageIsValid ImageIsValid;
+		FGetCurrentImageGuid GetCurrentImageGuid;
+	};
+}
+#else
+#pragma message("Do not include this header!")
+#endif
diff --git a/Engine/Plugins/Experimental/ImageWidgets/Source/ImageWidgets/Private/ImageViewportClient.cpp b/Engine/Plugins/Experimental/ImageWidgets/Source/ImageWidgets/Private/ImageViewportClient.cpp
--- a/Engine/Plugins/Experimental/ImageWidgets/Source/ImageWidgets/Private/ImageViewportClient.cpp
+++ b/Engine/Plugins/Experimental/ImageWidgets/Source/ImageWidgets/Private/ImageViewportClient.cpp
@@ -4,6 +4,9 @@
 
 #include "CanvasItem.h"
 #include "CanvasTypes.h"
+#if IMAGE_WIDGETS_WITH_AB_COMPARISON
+#include "ImageABComparison.h"
+#endif
 #include "ImageWidgetsLogCategory.h"
 #include "SImageViewport.h"
 #include "Texture2DPreview.h"
@@ -100,6 +103,9 @@ void DestroyCheckerTexture(TStrongObjectPtr<UTexture2D>& CheckerTexture)
 FImageViewportClient::FImageViewportClient(const TWeakPtr<SEditorViewport>& InEditorViewport, FGetImageSize&& InGetImageSize, FDrawImage&& InDrawImage,
                                            FGetDrawSettings&& InGetDrawSettings, FGetDPIScaleFactor&& InGetDPIScaleFactor,
                                            FOnLeftMouseButtonPressed&& InOnLeftMouseButtonPressed, FOnLeftMouseButtonReleased&& InOnLeftMouseButtonReleased,
+#if IMAGE_WIDGETS_WITH_AB_COMPARISON       
+                                           const FImageABComparison* ABComparison,
+#endif
                                            SImageViewport::FControllerSettings::EDefaultZoomMode DefaultZoomMode,
                                            EMouseCaptureMode InMouseCaptureMode /*= EMouseCaptureMode::CapturePermanently*/)
 	: FEditorViewportClient(nullptr, nullptr, InEditorViewport)
@@ -109,6 +115,9 @@ FImageViewportClient::FImageViewportClient(const TWeakPtr<SEditorViewport>& InEd
 	, GetDPIScaleFactor(MoveTemp(InGetDPIScaleFactor))
 	, OnLeftMouseButtonPressed(MoveTemp(InOnLeftMouseButtonPressed))
 	, OnLeftMouseButtonReleased(MoveTemp(InOnLeftMouseButtonReleased))
+#if IMAGE_WIDGETS_WITH_AB_COMPARISON
+	, ABComparison(ABComparison)
+#endif
 	, Controller(static_cast<FImageViewportController::EZoomMode>(DefaultZoomMode))
 	, MouseCaptureMode(InMouseCaptureMode)
 {
@@ -182,8 +191,44 @@ void FImageViewportClient::Draw(FViewport* InViewport, FCanvas* Canvas)
 		Canvas->DrawItem(Background);
 	}
 
+#if IMAGE_WIDGETS_WITH_AB_COMPARISON
+	const double ABComparisonDividerX = [this]
+	{
+		if (bDraggingABComparisonDivider)
+		{
+			FIntPoint MousePos;
+			Viewport->GetMousePos(MousePos);
+
+			FVector2d DividerPos = FVector2d(MousePos) - CachedPlacement.Offset;
+			DividerPos /= CachedPlacement.ZoomFactor;
+
+			const FIntPoint ImageSize = GetImageSize.Execute();
+			return DividerPos.X / ImageSize.X;
+		}
+		return ABComparisonDivider;
+	}();
+
+	const IImageViewer::FDrawProperties::FABComparison ABComparisonProperties = [this, &ABComparisonDividerX]() -> IImageViewer::FDrawProperties::FABComparison
+	{
+		return {ABComparison->GuidA(), ABComparison->GuidB(), FMath::Clamp(ABComparisonDividerX, 0.0, 1.0)};
+	}();
+
+	// Draw image
+	DrawImage.Execute(InViewport, Canvas, {CachedPlacement, MipProperties, ABComparisonProperties});
+
+	if (ABComparison->IsActive())
+	{
+		const FVector2D LineStart(CachedPlacement.Offset.X + CachedPlacement.Size.X * ABComparisonDividerX, 0);
+		const FVector2D LineEnd(CachedPlacement.Offset.X + CachedPlacement.Size.X * ABComparisonDividerX, InViewport->GetSizeXY().Y);
+		FCanvasLineItem Line(LineStart, LineEnd);
+		Line.LineThickness = 2.0f;
+		Line.SetColor({0.5f, 0.5f, 0.5f, 1.0f});
+		Canvas->DrawItem(Line);
+	}
+#else
 	// Draw image
 	DrawImage.Execute(InViewport, Canvas, {CachedPlacement, MipProperties});
+#endif
 }
 
 EMouseCursor::Type FImageViewportClient::GetCursor(FViewport* InViewport, int32 X, int32 Y)	
@@ -195,6 +240,13 @@ EMouseCursor::Type FImageViewportClient::GetCursor(FViewport* InViewport, int32
 		return EMouseCursor::GrabHandClosed;
 	}
 
+#if IMAGE_WIDGETS_WITH_AB_COMPARISON
+	if (bDraggingABComparisonDivider || MouseIsOverABComparisonDivider({X, Y}))
+	{
+		return EMouseCursor::ResizeLeftRight;
+	}
+#endif
+
 	const auto [PixelCoordsValid, PixelCoords] = GetPixelCoordinatesUnderCursor();
 	const FIntPoint ImageSize = GetImageSize.Execute();
 	if (PixelCoordsValid && 0 <= PixelCoords.X && 0 <= PixelCoords.Y && PixelCoords.X < ImageSize.X && PixelCoords.Y < ImageSize.Y)
@@ -259,6 +311,14 @@ void FImageViewportClient::TrackingStarted(const FInputEventState& InputState, b
 		FIntPoint MousePos;
 		InputState.GetViewport()->GetMousePos(MousePos);
 
+#if IMAGE_WIDGETS_WITH_AB_COMPARISON
+		if (InputState.IsLeftMouseButtonPressed() && MouseIsOverABComparisonDivider(MousePos))
+		{
+			bDraggingABComparisonDivider = true;
+			return;
+		}
+#endif
+
 		if (InputState.IsRightMouseButtonPressed())
 		{
 			bDragging = true;
@@ -272,6 +332,23 @@ void FImageViewportClient::TrackingStarted(const FInputEventState& InputState, b
 
 void FImageViewportClient::TrackingStopped()
 {
+#if IMAGE_WIDGETS_WITH_AB_COMPARISON
+	if (bDraggingABComparisonDivider)
+	{
+		bDraggingABComparisonDivider = false;
+
+		FIntPoint DraggingEnd;
+		Viewport->GetMousePos(DraggingEnd);
+
+		FVector2d DraggingEndPos = FVector2d(DraggingEnd) - CachedPlacement.Offset;
+		DraggingEndPos /= CachedPlacement.ZoomFactor;
+
+		const FIntPoint ImageSize = GetImageSize.Execute();
+		
+		ABComparisonDivider = DraggingEndPos.X / ImageSize.X;
+	}
+#endif
+
 	if (bDragging)
 	{
 		bDragging = false;
@@ -411,6 +488,14 @@ EMouseCaptureMode FImageViewportClient::GetMouseCaptureMode() const
 {
 	return MouseCaptureMode;
 }
+
+#if IMAGE_WIDGETS_WITH_AB_COMPARISON
+bool FImageViewportClient::MouseIsOverABComparisonDivider(const FIntPoint MousePos) const
+{
+	const double ABComparisonDividerPosition = CachedPlacement.Offset.X + CachedPlacement.Size.X * ABComparisonDivider;
+	return ABComparisonDividerPosition - 1.0 <= MousePos.X && MousePos.X <= ABComparisonDividerPosition + 1.0;
+}
+#endif
 }
 
 #undef LOCTEXT_NAMESPACE
diff --git a/Engine/Plugins/Experimental/ImageWidgets/Source/ImageWidgets/Private/ImageViewportClient.h b/Engine/Plugins/Experimental/ImageWidgets/Source/ImageWidgets/Private/ImageViewportClient.h
--- a/Engine/Plugins/Experimental/ImageWidgets/Source/ImageWidgets/Private/ImageViewportClient.h
+++ b/Engine/Plugins/Experimental/ImageWidgets/Source/ImageWidgets/Private/ImageViewportClient.h
@@ -9,7 +9,7 @@
 
 namespace UE::ImageWidgets
 {
-	class FABComparison;
+	class FImageABComparison;
 
 	DECLARE_DELEGATE_RetVal(FIntPoint, FGetImageSize)
 	DECLARE_DELEGATE_ThreeParams(FDrawImage, FViewport*, FCanvas*, const IImageViewer::FDrawProperties&)
@@ -27,6 +27,9 @@ namespace UE::ImageWidgets
 		FImageViewportClient(const TWeakPtr<SEditorViewport>& InViewport, FGetImageSize&& InGetImageSize, FDrawImage&& InDrawImage,
 		                     FGetDrawSettings&& InGetDrawSettings, FGetDPIScaleFactor&& InGetDPIScaleFactor,
 		                     FOnLeftMouseButtonPressed&& InOnLeftMouseButtonPressed, FOnLeftMouseButtonReleased&& InOnLeftMouseButtonReleased,
+#if IMAGE_WIDGETS_WITH_AB_COMPARISON
+							 const FImageABComparison* ABComparison,
+#endif
 		                     SImageViewport::FControllerSettings::EDefaultZoomMode DefaultZoomMode, EMouseCaptureMode InMouseCaptureMode = EMouseCaptureMode::CapturePermanently);
 		virtual ~FImageViewportClient() override;
 
@@ -65,6 +68,9 @@ namespace UE::ImageWidgets
 		IImageViewer::FDrawProperties::FMip GetMipProperties() const;
 
 		void CreateOrDestroyCheckerTextureIfSettingsChanged(const SImageViewport::FDrawSettings& DrawSettings);
+#if IMAGE_WIDGETS_WITH_AB_COMPARISON
+		bool MouseIsOverABComparisonDivider(FIntPoint MousePos) const;
+#endif
 
 		FVector2d GetViewportSizeWithDPIScaling() const;
 		
@@ -74,6 +80,9 @@ namespace UE::ImageWidgets
 		FGetDPIScaleFactor GetDPIScaleFactor;
 		FOnLeftMouseButtonPressed OnLeftMouseButtonPressed;
 		FOnLeftMouseButtonReleased OnLeftMouseButtonReleased;
+#if IMAGE_WIDGETS_WITH_AB_COMPARISON
+		const FImageABComparison* ABComparison;
+#endif
 		
 		bool bDragging = false;
 		FIntPoint DraggingStart;
@@ -89,6 +98,11 @@ namespace UE::ImageWidgets
 		TStrongObjectPtr<UTexture2D> CheckerTexture;
 		FCheckerTextureSettings CachedCheckerTextureSettings;
 
+#if IMAGE_WIDGETS_WITH_AB_COMPARISON
+		double ABComparisonDivider = 0.5;
+		bool bDraggingABComparisonDivider = false;
+#endif
+		
 		EMouseCaptureMode MouseCaptureMode;
 	};
 }
diff --git a/Engine/Plugins/Experimental/ImageWidgets/Source/ImageWidgets/Private/SImageViewport.cpp b/Engine/Plugins/Experimental/ImageWidgets/Source/ImageWidgets/Private/SImageViewport.cpp
--- a/Engine/Plugins/Experimental/ImageWidgets/Source/ImageWidgets/Private/SImageViewport.cpp
+++ b/Engine/Plugins/Experimental/ImageWidgets/Source/ImageWidgets/Private/SImageViewport.cpp
@@ -3,6 +3,9 @@
 #include "SImageViewport.h"
 
 #include "CanvasTypes.h"
+#if IMAGE_WIDGETS_WITH_AB_COMPARISON
+#include "ImageABComparison.h"
+#endif
 #include "ImageViewportClient.h"
 #include "ImageViewportController.h"
 #include "ImageWidgetsCommands.h"
@@ -49,6 +52,11 @@ namespace UE::ImageWidgets
 	}
 
 	SImageViewport::SImageViewport()
+#if IMAGE_WIDGETS_WITH_AB_COMPARISON
+		: ABComparison(MakePimpl<FImageABComparison>(
+			FImageABComparison::FImageIsValid::CreateLambda([this](const FGuid& Guid) { return ImageViewer->IsValidImage(Guid); }),
+			FImageABComparison::FGetCurrentImageGuid::CreateLambda([this] { return ImageViewer->GetCurrentImageInfo().Guid; })))
+#endif
 	{
 	}
 
@@ -148,6 +156,9 @@ namespace UE::ImageWidgets
 		StatusBarExtender = InArgs._StatusBarExtender;
 
 		DrawSettings = InArgs._DrawSettings;
+#if IMAGE_WIDGETS_WITH_AB_COMPARISON
+		bABComparisonEnabled = InArgs._bABComparisonEnabled;
+#endif
 		OnLeftMouseButtonPressed = InArgs._OnLeftMouseButtonPressed;
 		OnLeftMouseButtonReleased = InArgs._OnLeftMouseButtonReleased;
 
@@ -189,7 +200,7 @@ namespace UE::ImageWidgets
 		{
 			OnLeftMouseButtonReleased.ExecuteIfBound();
 		};
-		
+
 		ImageViewportClient = MakeShareable(new FImageViewportClient(
 			StaticCastWeakPtr<SEditorViewport>(AsWeak()),
 			FGetImageSize::CreateLambda(GetImageSize),
@@ -198,6 +209,9 @@ namespace UE::ImageWidgets
 			FGetDPIScaleFactor::CreateLambda(GetDPIScaleFactor),
 			FOnLeftMouseButtonPressed::CreateLambda(HandleOnLeftMouseButtonPressed),
 			FOnLeftMouseButtonReleased::CreateLambda(HandleOnLeftMouseButtonReleased),
+#if IMAGE_WIDGETS_WITH_AB_COMPARISON
+			ABComparison.Get(),
+#endif
 			InArgs._ControllerSettings.DefaultZoomMode,
 			InArgs._MouseCaptureMode)
 			);
@@ -286,6 +300,17 @@ namespace UE::ImageWidgets
 			return ImageViewer->GetCurrentImageInfo().Guid;
 		};
 
+#if IMAGE_WIDGETS_WITH_AB_COMPARISON
+		SAssignNew(ImageViewportToolbar, SImageViewportToolbar, StaticCastSharedPtr<FImageViewportClient>(Client), CommandList,
+						  SImageViewportToolbar::FConstructParameters
+						  {
+							  SImageViewportToolbar::FHasImage::CreateLambda(HasImage),
+							  SImageViewportToolbar::FNumMips::CreateLambda(GetNumMips),
+							  SImageViewportToolbar::FImageGuid::CreateLambda(GetImageGuid),
+							  bABComparisonEnabled ? ABComparison.Get() : nullptr,
+							  ToolbarExtender
+						  });
+#else
 		SAssignNew(ImageViewportToolbar, SImageViewportToolbar, StaticCastSharedPtr<FImageViewportClient>(Client), CommandList,
 						  SImageViewportToolbar::FConstructParameters
 						  {
@@ -294,6 +319,8 @@ namespace UE::ImageWidgets
 							  SImageViewportToolbar::FImageGuid::CreateLambda(GetImageGuid),
 							  ToolbarExtender
 						  });
+#endif
+						  
 
 		ToolbarExtender.Reset();
 
diff --git a/Engine/Plugins/Experimental/ImageWidgets/Source/ImageWidgets/Private/SImageViewportToolbar.cpp b/Engine/Plugins/Experimental/ImageWidgets/Source/ImageWidgets/Private/SImageViewportToolbar.cpp
--- a/Engine/Plugins/Experimental/ImageWidgets/Source/ImageWidgets/Private/SImageViewportToolbar.cpp
+++ b/Engine/Plugins/Experimental/ImageWidgets/Source/ImageWidgets/Private/SImageViewportToolbar.cpp
@@ -2,6 +2,9 @@
 
 #include "SImageViewportToolbar.h"
 
+#if IMAGE_WIDGETS_WITH_AB_COMPARISON
+#include "ImageABComparison.h"
+#endif
 #include "ImageViewportClient.h"
 #include "ImageWidgetsCommands.h"
 #include "SEditorViewportToolBarMenu.h"
@@ -35,6 +38,10 @@ namespace UE::ImageWidgets
 		check(NumMips.IsBound());
 		check(ImageGuid.IsBound());
 
+#if IMAGE_WIDGETS_WITH_AB_COMPARISON
+		ABComparison = Parameters.ABComparison;
+#endif
+
 		ChildSlot
 		[
 			SNew(SBorder)
@@ -51,7 +58,11 @@ namespace UE::ImageWidgets
 				  .FillWidth(1.0f)
 				  .HAlign(HAlign_Center)
 				[
+#if IMAGE_WIDGETS_WITH_AB_COMPARISON
+MakeCenterToolbar(Parameters.ToolbarExtender, ABComparison != nullptr)
+#else
 					MakeCenterToolbar(Parameters.ToolbarExtender)
+#endif
 				]
 				+ SHorizontalBox::Slot()
 				  .AutoWidth()
@@ -110,6 +121,65 @@ namespace UE::ImageWidgets
 		return ToolbarBuilder.MakeWidget();
 	}
 
+	#if IMAGE_WIDGETS_WITH_AB_COMPARISON
+	TSharedRef<SWidget> SImageViewportToolbar::MakeCenterToolbar(const TSharedPtr<FExtender>& Extender, bool bEnableABComparison)
+	{
+		FSlimHorizontalToolBarBuilder ToolbarBuilder = GetToolbarBuilder(CommandList, Extender);
+
+		ToolbarBuilder.BeginSection("ToolbarCenter");
+		{
+			ToolbarBuilder.AddSeparator();
+
+			if (bEnableABComparison && ABComparison)
+			{
+				auto GetTextButton = [this](const FString& Label, const FCheckBoxStyle* ButtonStyle, FImageABComparison::EAorB AorB)
+				{
+					return SNew(SCheckBox)
+					.Style(ButtonStyle)
+					.IsEnabled_Lambda([this, AorB]
+					                      {
+						                      return ABComparison->CanSetABComparison(AorB);
+					                      })
+					.IsChecked(ABComparison->ABComparisonIsSet(AorB))
+					.OnCheckStateChanged_Lambda([this, AorB](const ECheckBoxState State)
+					                      {
+						                      ABComparison->SetABComparison(AorB, State != ECheckBoxState::Checked ? FGuid() : ImageGuid.Execute());
+					                      })
+					[
+						SNew(STextBlock)
+							.Font(FAppStyle::GetFontStyle("EditorViewportToolBar.Font"))
+							.Text(FText::FromString(Label))
+							.Margin(FMargin(2.0f, 0.0f))
+					];
+				};
+
+				const FCheckBoxStyle* ButtonStyleStart = &FAppStyle::Get().GetWidgetStyle<FCheckBoxStyle>("EditorViewportToolBar.ToggleButton.Start");
+				const FCheckBoxStyle* ButtonStyleEnd = &FAppStyle::Get().GetWidgetStyle<FCheckBoxStyle>("EditorViewportToolBar.ToggleButton.End");
+
+				const TSharedRef<SHorizontalBox> RGBA = SNew(SHorizontalBox)
+					+ SHorizontalBox::Slot()
+					.AutoWidth()
+					[
+						GetTextButton("A", ButtonStyleStart, FImageABComparison::A)
+					]
+					+ SHorizontalBox::Slot()
+					.AutoWidth()
+					[
+						GetTextButton("B", ButtonStyleEnd, FImageABComparison::B)
+					];
+
+				ToolbarBuilder.AddToolBarWidget(RGBA);
+			}
+
+			ToolbarBuilder.AddSeparator();
+		}
+		ToolbarBuilder.EndSection();
+
+		ToolbarBuilder.AddSeparator();
+
+		return ToolbarBuilder.MakeWidget();
+	}
+#else
 	TSharedRef<SWidget> SImageViewportToolbar::MakeCenterToolbar(const TSharedPtr<FExtender>& Extender)
 	{
 		FSlimHorizontalToolBarBuilder ToolbarBuilder = GetToolbarBuilder(CommandList, Extender);
@@ -125,6 +195,7 @@ namespace UE::ImageWidgets
 
 		return ToolbarBuilder.MakeWidget();
 	}
+#endif
 
 	TSharedRef<SWidget> SImageViewportToolbar::MakeRightToolbar(const TSharedPtr<FExtender>& Extender)
 	{
diff --git a/Engine/Plugins/Experimental/ImageWidgets/Source/ImageWidgets/Private/SImageViewportToolbar.h b/Engine/Plugins/Experimental/ImageWidgets/Source/ImageWidgets/Private/SImageViewportToolbar.h
--- a/Engine/Plugins/Experimental/ImageWidgets/Source/ImageWidgets/Private/SImageViewportToolbar.h
+++ b/Engine/Plugins/Experimental/ImageWidgets/Source/ImageWidgets/Private/SImageViewportToolbar.h
@@ -13,6 +13,9 @@ class SViewportToolBarComboMenu;
 
 namespace UE::ImageWidgets
 {
+#if IMAGE_WIDGETS_WITH_AB_COMPARISON
+	class FImageABComparison;
+#endif
 	class FImageViewportClient;
 
 	/**
@@ -36,6 +39,9 @@ namespace UE::ImageWidgets
 			FHasImage HasImage;
 			FNumMips NumMips;
 			FImageGuid ImageGuid;
+#if IMAGE_WIDGETS_WITH_AB_COMPARISON
+			FImageABComparison* ABComparison = nullptr;
+#endif
 			TSharedPtr<FExtender> ToolbarExtender;
 		};
 		
@@ -44,7 +50,11 @@ namespace UE::ImageWidgets
 
 	private:
 		TSharedRef<SWidget> MakeLeftToolbar(const TSharedPtr<FExtender>& Extender);
+#if IMAGE_WIDGETS_WITH_AB_COMPARISON
+		TSharedRef<SWidget> MakeCenterToolbar(const TSharedPtr<FExtender>& Extender, bool bEnableABComparison);
+#else
 		TSharedRef<SWidget> MakeCenterToolbar(const TSharedPtr<FExtender>& Extender);
+#endif
 		TSharedRef<SWidget> MakeRightToolbar(const TSharedPtr<FExtender>& Extender);
 
 		EVisibility GetMipMenuVisibility() const;
@@ -60,5 +70,8 @@ namespace UE::ImageWidgets
 		FHasImage HasImage;
 		FNumMips NumMips;
 		FImageGuid ImageGuid;
+#if IMAGE_WIDGETS_WITH_AB_COMPARISON
+		FImageABComparison* ABComparison = nullptr;
+#endif
 	};
 }
diff --git a/Engine/Plugins/Experimental/ImageWidgets/Source/ImageWidgets/Public/IImageViewer.h b/Engine/Plugins/Experimental/ImageWidgets/Source/ImageWidgets/Public/IImageViewer.h
--- a/Engine/Plugins/Experimental/ImageWidgets/Source/ImageWidgets/Public/IImageViewer.h
+++ b/Engine/Plugins/Experimental/ImageWidgets/Source/ImageWidgets/Public/IImageViewer.h
@@ -68,14 +68,42 @@ namespace UE::ImageWidgets
 				float MipLevel;
 			};
 
+#if IMAGE_WIDGETS_WITH_AB_COMPARISON
+			/**
+			 * Information necessary for rendering AB comparisons.
+			 */
+			struct FABComparison
+			{
+				/** @return true if an AB comparison should be drawn instead of a single image. */
+				bool IsActive() const { return GuidA.IsValid() && GuidB.IsValid(); };
+
+				/** Unique identifier for image A. */
+				FGuid GuidA;
+				
+				/** Unique identifier for image B. */
+				FGuid GuidB;
+
+				/**
+				 * Value between 0..1 to indicate where the threshold between images A and B is.
+				 * A value of 0 means that only B should be drawn.
+				 * A value of 0.5 means that the left half of A and the right half of B should be drawn.
+				 * A value of 1 means that only A should be drawn. 
+				 */
+				double Threshold;
+			};
+#endif
+
 			FPlacement Placement;
 			FMip Mip;
+#if IMAGE_WIDGETS_WITH_AB_COMPARISON
+			FABComparison ABComparison;
+#endif
 		};
 
 		/**
 		 * Provides any necessary metadata for the image widgets about the image that is currently supposed to be displayed. This data is generic in the sense
 		 * that the image widgets don't need to know any of the image structure, its content or how to draw it. Instead, the image drawing is done directly
-		 * via @ref DrawCurrentImage.
+		 * via @see DrawCurrentImage.
 		 * @return information about the image to be displayed
 		 */
 		virtual FImageInfo GetCurrentImageInfo() const = 0;
@@ -96,13 +124,20 @@ namespace UE::ImageWidgets
 		 */
 		virtual TOptional<TVariant<FColor, FLinearColor>> GetCurrentImagePixelColor(FIntPoint PixelCoords, int32 MipLevel) const = 0;
 
-#if IMAGE_WIDGETS_WITH_CATALOG
 		/**
 		 * Notifies about the image with the given GUID being selected.
 		 * This can be implemented as an empty function if the image viewer implementation does not support switching between different images.
 		 * @param Guid Unique identifier of the selected image
 		 */
 		virtual void OnImageSelected(const FGuid& Guid) = 0;
+
+#if IMAGE_WIDGETS_WITH_AB_COMPARISON
+		/**
+		 * Returns if a given GUID represents a currently available image.
+		 * @param Guid Unique identifier of the potentially available image
+		 * @return true if the GUID represents a currently available image
+		 */
+		virtual bool IsValidImage(const FGuid& Guid) const = 0;
 #endif
 
 	protected:
diff --git a/Engine/Plugins/Experimental/ImageWidgets/Source/ImageWidgets/Public/SImageCatalog.h b/Engine/Plugins/Experimental/ImageWidgets/Source/ImageWidgets/Public/SImageCatalog.h
--- a/Engine/Plugins/Experimental/ImageWidgets/Source/ImageWidgets/Public/SImageCatalog.h
+++ b/Engine/Plugins/Experimental/ImageWidgets/Source/ImageWidgets/Public/SImageCatalog.h
@@ -36,7 +36,7 @@ namespace UE::ImageWidgets
 
 	/**
 	 * Generic catalog widget for listing and selecting 2D image-like content.
-	 * Each catalog item is represented by its own widget based on its @ref FImageCatalogItemData.
+	 * Each catalog item is represented by its own widget based on its @see FImageCatalogItemData.
 	 * Entries in the catalog can be split into regular and pinned items, with pinned items being displayed at the top.
 	 */
 	class SImageCatalog : public SCompoundWidget
@@ -44,7 +44,7 @@ namespace UE::ImageWidgets
 	public:
 		/**
 		 * Delegate that gets called when an item is selected in the catalog.
-		 * The given @ref FGuid identifies the item that was selected.
+		 * The given @see FGuid identifies the item that was selected.
 		 */
 		DECLARE_DELEGATE_OneParam(FOnItemSelected, const FGuid&)
 
diff --git a/Engine/Plugins/Experimental/ImageWidgets/Source/ImageWidgets/Public/SImageViewport.h b/Engine/Plugins/Experimental/ImageWidgets/Source/ImageWidgets/Public/SImageViewport.h
--- a/Engine/Plugins/Experimental/ImageWidgets/Source/ImageWidgets/Public/SImageViewport.h
+++ b/Engine/Plugins/Experimental/ImageWidgets/Source/ImageWidgets/Public/SImageViewport.h
@@ -10,6 +10,9 @@ class SViewportToolBar;
 
 namespace UE::ImageWidgets
 {
+#if IMAGE_WIDGETS_WITH_AB_COMPARISON
+	class FImageABComparison;
+#endif
 	class FImageViewportClient;
 	class FStatusBarExtension;
 	class IImageViewer;
@@ -17,7 +20,7 @@ namespace UE::ImageWidgets
 
 	/**
 	 * Generic viewport for displaying and interacting with 2D image-like content.
-	 * The drawing of the images is deferred to an @ref IImageViewer implementation that needs to be provided upon construction. This viewport only uses the 
+	 * The drawing of the images is deferred to an @see IImageViewer implementation that needs to be provided upon construction. This viewport only uses the 
 	 * metadata provided by the image viewer to have sufficient information about the image without being aware of its actual format or contents.
 	 */
 	class SImageViewport : public SEditorViewport
@@ -33,8 +36,8 @@ namespace UE::ImageWidgets
 		public:
 			/**
 			 * Delegate that gets called for each extension when the status bar is constructed.
-			 * The given @ref SHorizontalBox is the layout container that the widget for the extension needs to be added to by calling
-			 * @ref SHorizontalBox::AddSlot().
+			 * The given @see SHorizontalBox is the layout container that the widget for the extension needs to be added to by calling
+			 * @see SHorizontalBox::AddSlot().
 			 */
 			DECLARE_DELEGATE_OneParam(FDelegate, SHorizontalBox&);
 
@@ -107,15 +110,20 @@ namespace UE::ImageWidgets
 				Fill // Make the image fit within the viewport, and if it is smaller than the viewport, zoom in to fill the viewport.
 			};
 
-			/** Zoom mode that gets set on viewport construction and whenever the viewport controller is reset via @ref ResetController(). */
+			/** Zoom mode that gets set on viewport construction and whenever the viewport controller is reset via @see ResetController(). */
 			EDefaultZoomMode DefaultZoomMode = EDefaultZoomMode::Fit;
 		};
 
 		DECLARE_DELEGATE(FOnLeftMouseButtonPressed);
 		DECLARE_DELEGATE(FOnLeftMouseButtonReleased);
 
-		SLATE_BEGIN_ARGS(SImageViewport) :
-			_MouseCaptureMode(EMouseCaptureMode::CapturePermanently)
+		SLATE_BEGIN_ARGS(SImageViewport)
+#if IMAGE_WIDGETS_WITH_AB_COMPARISON
+			: _bABComparisonEnabled(false)
+			, _MouseCaptureMode(EMouseCaptureMode::CapturePermanently)
+#else
+			: _MouseCaptureMode(EMouseCaptureMode::CapturePermanently)
+#endif
 			{
 			}
 
@@ -128,6 +136,11 @@ namespace UE::ImageWidgets
 			/** Settings for drawing viewport contents other than the actual image */
 			SLATE_ATTRIBUTE(FDrawSettings, DrawSettings)
 
+#if IMAGE_WIDGETS_WITH_AB_COMPARISON
+			/** Enables AB comparison controls in the toolbar */
+			SLATE_ARGUMENT(bool, bABComparisonEnabled)
+#endif
+
 			/** Settings for controlling the viewport */
 			SLATE_ARGUMENT(FControllerSettings, ControllerSettings)
 
@@ -148,19 +161,19 @@ namespace UE::ImageWidgets
 		/**
 		 * Function used by Slate to construct the image viewport widget with the given arguments.
 		 * @param InArgs Slate arguments defined above
-		 * @param InImageViewer @ref IImageViewer implementation that holds and draws the actual image contents.
+		 * @param InImageViewer @see IImageViewer implementation that holds and draws the actual image contents.
 		 */
 		IMAGEWIDGETS_API void Construct(const FArguments& InArgs, const TSharedRef<IImageViewer>& InImageViewer);
 
 		/**
-		 * Provides access to the viewport toolbar, which is needed to dynamically generate certain toolbar widgets, e.g. an @ref SEditorViewportToolbarMenu,
+		 * Provides access to the viewport toolbar, which is needed to dynamically generate certain toolbar widgets, e.g. an @see SEditorViewportToolbarMenu,
 		 * as part of a toolbar extension. 
 		 * @return Pointer to the viewport toolbar 
 		 */
 		IMAGEWIDGETS_API TSharedPtr<SViewportToolBar> GetParentToolbar() const;
 
 		/**
-		 * Result for calls to @ref GetPixelCoordinatesUnderCursor.
+		 * Result for calls to @see GetPixelCoordinatesUnderCursor.
 		 */
 		struct FPixelCoordinatesUnderCursorResult
 		{
@@ -174,7 +187,7 @@ namespace UE::ImageWidgets
 
 		/**
 		 * Provides the pixel coordinates under the cursor.
-		 * @return Pixel coordinates result under the cursor; see @ref FPixelCoordinatesUnderCursorResult for more details
+		 * @return Pixel coordinates result under the cursor; see @see FPixelCoordinatesUnderCursorResult for more details
 		 */
 		IMAGEWIDGETS_API FPixelCoordinatesUnderCursorResult GetPixelCoordinatesUnderCursor() const;
 
@@ -218,6 +231,14 @@ namespace UE::ImageWidgets
 		/** Makes the draw setting available either as fixed values or via a callback to the outside of the viewport. */
 		TAttribute<FDrawSettings> DrawSettings;
 
+#if IMAGE_WIDGETS_WITH_AB_COMPARISON
+		/** Flag that determines is AB comparison widgets are enabled or not. The value does not change after the call to @see Construct(). */
+		bool bABComparisonEnabled = false;
+
+		/** Data and logic related to AB comparisons. This is effectively unused when @see bABComparisonEnabled is set to false. */
+		TPimplPtr<FImageABComparison> ABComparison;
+#endif
+
 		/** The image viewer that holds and draws the actual images. */
 		TSharedPtr<IImageViewer> ImageViewer;
 
@@ -227,10 +248,10 @@ namespace UE::ImageWidgets
 		/** The toolbar that controls some of the behavior of the viewport and optionally also the image viewer via toolbar extensions. */
 		TSharedPtr<SImageViewportToolbar> ImageViewportToolbar;
 
-		/** Toolbar extensions provided by the call to @ref Construct(). This pointer is reset after the extensions were applied during construction. */
+		/** Toolbar extensions provided by the call to @see Construct(). This pointer is reset after the extensions were applied during construction. */
 		TSharedPtr<FExtender> ToolbarExtender;
 
-		/** Status bar extensions provided by the call to @ref Construct(). This pointer is reset after the extensions were applied during construction. */
+		/** Status bar extensions provided by the call to @see Construct(). This pointer is reset after the extensions were applied during construction. */
 		TSharedPtr<FStatusBarExtender> StatusBarExtender;
 
 		/** Left Mouse buttons pressed and released events triggered to tell the widget containing the viewport. These are initialized from the constructor.*/
diff --git a/Engine/Plugins/Experimental/TextureGraph/Source/TextureGraphEditor/Private/STG_NodePreview.h b/Engine/Plugins/Experimental/TextureGraph/Source/TextureGraphEditor/Private/STG_NodePreview.h
--- a/Engine/Plugins/Experimental/TextureGraph/Source/TextureGraphEditor/Private/STG_NodePreview.h
+++ b/Engine/Plugins/Experimental/TextureGraph/Source/TextureGraphEditor/Private/STG_NodePreview.h
@@ -29,9 +29,7 @@ public:
 	virtual FImageInfo GetCurrentImageInfo() const override;
 	virtual void DrawCurrentImage(FViewport* Viewport, FCanvas* Canvas, const FDrawProperties& Properties) override;
 	virtual TOptional<TVariant<FColor, FLinearColor>> GetCurrentImagePixelColor(FIntPoint PixelCoords, int32 MipIndex) const override;
-#if IMAGE_WIDGETS_WITH_CATALOG
 	virtual void OnImageSelected(const FGuid& ImageGuid) override {};
-#endif
 #if IMAGE_WIDGETS_WITH_AB_COMPARISON
 	virtual bool IsValidImage(const FGuid& Guid) const override { return true; }
 #endif
