ImageWidgets: AB comparison tooltips

diff --git a/Engine/Plugins/Experimental/ImageWidgets/Source/ImageWidgets/Private/ColorViewerSample/ColorViewer.cpp b/Engine/Plugins/Experimental/ImageWidgets/Source/ImageWidgets/Private/ColorViewerSample/ColorViewer.cpp
--- a/Engine/Plugins/Experimental/ImageWidgets/Source/ImageWidgets/Private/ColorViewerSample/ColorViewer.cpp
+++ b/Engine/Plugins/Experimental/ImageWidgets/Source/ImageWidgets/Private/ColorViewerSample/ColorViewer.cpp
@@ -61,6 +61,17 @@ namespace UE::ImageWidgets::Sample
 	{
 		return ColorIsValid(Guid.B) && Colors[Guid.B].Guid == Guid;
 	}
+
+	FText FColorViewer::GetImageName(const FGuid& Guid) const
+	{
+		if (IsValidImage(Guid))
+		{
+			const FColor& Color = Colors[Guid.B].Color;
+			const FString HexColor = FString::Printf(TEXT("#%02X%02X%02X"), Color.R, Color.G, Color.B);
+			return FText::FromString(HexColor);
+		}
+		return {};
+	}
 #endif
 
 	const FColorViewer::FColorItem* FColorViewer::AddColor()
diff --git a/Engine/Plugins/Experimental/ImageWidgets/Source/ImageWidgets/Private/ColorViewerSample/ColorViewer.h b/Engine/Plugins/Experimental/ImageWidgets/Source/ImageWidgets/Private/ColorViewerSample/ColorViewer.h
--- a/Engine/Plugins/Experimental/ImageWidgets/Source/ImageWidgets/Private/ColorViewerSample/ColorViewer.h
+++ b/Engine/Plugins/Experimental/ImageWidgets/Source/ImageWidgets/Private/ColorViewerSample/ColorViewer.h
@@ -63,6 +63,7 @@ namespace UE::ImageWidgets::Sample
 		virtual void OnImageSelected(const FGuid& Guid) override;
 #if IMAGE_WIDGETS_WITH_AB_COMPARISON
 		virtual bool IsValidImage(const FGuid& Guid) const override;
+		virtual FText GetImageName(const FGuid& Guid) const override;
 #endif
 		// IImageViewer overrides - end
 
diff --git a/Engine/Plugins/Experimental/ImageWidgets/Source/ImageWidgets/Private/ColorViewerSample/ColorViewerWidget.cpp b/Engine/Plugins/Experimental/ImageWidgets/Source/ImageWidgets/Private/ColorViewerSample/ColorViewerWidget.cpp
--- a/Engine/Plugins/Experimental/ImageWidgets/Source/ImageWidgets/Private/ColorViewerSample/ColorViewerWidget.cpp
+++ b/Engine/Plugins/Experimental/ImageWidgets/Source/ImageWidgets/Private/ColorViewerSample/ColorViewerWidget.cpp
@@ -98,16 +98,20 @@ namespace UE::ImageWidgets::Sample
 		ToolbarBuilder.EndBlockGroup();
 	}
 
-	TTuple<FText, FText, FText> GetColorItemMetaData(const FColor Color, FDateTime DateTime)
+	TTuple<FText, FText, FText> GetColorItemMetaData(const FColorViewer& ColorViewer, const FColorViewer::FColorItem* ColorItem)
 	{
-		const FString HexColor = FString::Printf(TEXT("#%02X%02X%02X"), Color.R, Color.G, Color.B);
-		const FText Name = FText::Format(LOCTEXT("ColorEntryLabel", "{0}"), FText::FromString(HexColor));
+#if IMAGE_WIDGETS_WITH_AB_COMPARISON
+		const FText Name = ColorViewer.GetImageName(ColorItem->Guid);
+#else
+		const FString HexColor = FString::Printf(TEXT("#%02X%02X%02X"), ColorItem->Color.R, ColorItem->Color.G, ColorItem->Color.B);
+		const FText Name = FText::FromString(HexColor);
+#endif
 
 		const FText Info = FText::Format(
-			LOCTEXT("ColorEntryInfoLabel", "{0}"), FText::AsTime(DateTime, EDateTimeStyle::Short, FText::GetInvariantTimeZone()));
+			LOCTEXT("ColorEntryInfoLabel", "{0}"), FText::AsTime(ColorItem->DateTime, EDateTimeStyle::Short, FText::GetInvariantTimeZone()));
 
 		const FText ToolTip = FText::Format(
-			LOCTEXT("ColorEntryToolTip", "R {0}, G {1}, B {2}"), {Color.R, Color.G, Color.B});
+			LOCTEXT("ColorEntryToolTip", "R {0}, G {1}, B {2}"), {ColorItem->Color.R, ColorItem->Color.G, ColorItem->Color.B});
 
 		return {Name, Info, ToolTip};
 	}
@@ -116,7 +120,7 @@ namespace UE::ImageWidgets::Sample
 	{
 		if (const FColorViewer::FColorItem* ColorItem = ColorViewer->AddColor())
 		{
-			auto [Name, Info, ToolTip] = GetColorItemMetaData(ColorItem->Color, ColorItem->DateTime);
+			auto [Name, Info, ToolTip] = GetColorItemMetaData(*ColorViewer, ColorItem);
 			const FLinearColor ToneMappedColor = ColorViewer->GetDefaultToneMappedColor(ColorItem->Color);
 
 			Catalog->AddItem(MakeShared<FImageCatalogItemData>(ColorItem->Guid, FSlateColorBrush(ToneMappedColor), Name, Info, ToolTip));
@@ -134,9 +138,10 @@ namespace UE::ImageWidgets::Sample
 	{
 		if (const FColorViewer::FColorItem* ColorItem = ColorViewer->RandomizeColor())
 		{
-			const auto [Name, Info, ToolTip] = GetColorItemMetaData(ColorItem->Color, ColorItem->DateTime);
+			const auto [Name, Info, ToolTip] = GetColorItemMetaData(*ColorViewer, ColorItem);
+			const FLinearColor ToneMappedColor = ColorViewer->GetDefaultToneMappedColor(ColorItem->Color);
 
-			Catalog->UpdateItem({ColorItem->Guid, FSlateColorBrush(ColorItem->Color), Name, Info, ToolTip});
+			Catalog->UpdateItem({ColorItem->Guid, FSlateColorBrush(ToneMappedColor), Name, Info, ToolTip});
 		}
 	}
 
diff --git a/Engine/Plugins/Experimental/ImageWidgets/Source/ImageWidgets/Private/ImageABComparison.h b/Engine/Plugins/Experimental/ImageWidgets/Source/ImageWidgets/Private/ImageABComparison.h
--- a/Engine/Plugins/Experimental/ImageWidgets/Source/ImageWidgets/Private/ImageABComparison.h
+++ b/Engine/Plugins/Experimental/ImageWidgets/Source/ImageWidgets/Private/ImageABComparison.h
@@ -2,10 +2,12 @@
 
 #pragma once
 
+#if IMAGE_WIDGETS_WITH_AB_COMPARISON
+
 #include "Delegates/Delegate.h"
+#include "Internationalization/Text.h"
 #include "Misc/Guid.h"
 
-#if IMAGE_WIDGETS_WITH_AB_COMPARISON
 namespace UE::ImageWidgets
 {
 	/**
@@ -16,16 +18,19 @@ namespace UE::ImageWidgets
 	public:
 		DECLARE_DELEGATE_RetVal(FGuid, FGetCurrentImageGuid)
 		DECLARE_DELEGATE_RetVal_OneParam(bool, FImageIsValid, const FGuid&)
+		DECLARE_DELEGATE_RetVal_OneParam(FText, FGetImageName, const FGuid&)
 
-		FImageABComparison(FImageIsValid&& ImageIsValid, FGetCurrentImageGuid&& GetCurrentImageGuid)
-			: ImageIsValid(ImageIsValid)
-			, GetCurrentImageGuid(GetCurrentImageGuid)
+		FImageABComparison(FImageIsValid&& InImageIsValid, FGetCurrentImageGuid&& InGetCurrentImageGuid, FGetImageName&& InGetImageName)
+			: ImageIsValid(MoveTemp(InImageIsValid))
+			, GetCurrentImageGuid(MoveTemp(InGetCurrentImageGuid))
+			, GetImageName(MoveTemp(InGetImageName))
 		{
 			check(ImageIsValid.IsBound());
 			check(GetCurrentImageGuid.IsBound());
+			check(GetImageName.IsBound());
 		}
 
-		enum EAorB
+		enum class EAorB : int8
 		{
 			A,
 			B
@@ -53,23 +58,29 @@ namespace UE::ImageWidgets
 
 		bool IsActive() const
 		{
-			return Guids[0].IsValid() && Guids[1].IsValid();
+			return Guids[static_cast<int32>(EAorB::A)].IsValid() && Guids[static_cast<int32>(EAorB::B)].IsValid();
+		}
+
+		const FGuid& GetGuidA() const
+		{
+			return Guids[static_cast<int32>(EAorB::A)];
 		}
 
-		const FGuid& GuidA() const
+		const FGuid& GetGuidB() const
 		{
-			return Guids[0];
+			return Guids[static_cast<int32>(EAorB::B)];
 		}
 
-		const FGuid& GuidB() const
+		FText GetName(EAorB AorB) const
 		{
-			return Guids[1];
+			return Guids[static_cast<int32>(AorB)].IsValid() ? GetImageName.Execute(Guids[static_cast<int32>(AorB)]) : FText();
 		}
 
 	private:
 		FGuid Guids[2];
 		FImageIsValid ImageIsValid;
 		FGetCurrentImageGuid GetCurrentImageGuid;
+		FGetImageName GetImageName;
 	};
 }
 #else
diff --git a/Engine/Plugins/Experimental/ImageWidgets/Source/ImageWidgets/Private/ImageViewportClient.cpp b/Engine/Plugins/Experimental/ImageWidgets/Source/ImageWidgets/Private/ImageViewportClient.cpp
--- a/Engine/Plugins/Experimental/ImageWidgets/Source/ImageWidgets/Private/ImageViewportClient.cpp
+++ b/Engine/Plugins/Experimental/ImageWidgets/Source/ImageWidgets/Private/ImageViewportClient.cpp
@@ -210,7 +210,7 @@ void FImageViewportClient::Draw(FViewport* InViewport, FCanvas* Canvas)
 
 	const IImageViewer::FDrawProperties::FABComparison ABComparisonProperties = [this, &ABComparisonDividerX]() -> IImageViewer::FDrawProperties::FABComparison
 	{
-		return {ABComparison->GuidA(), ABComparison->GuidB(), FMath::Clamp(ABComparisonDividerX, 0.0, 1.0)};
+		return {ABComparison->GetGuidA(), ABComparison->GetGuidB(), FMath::Clamp(ABComparisonDividerX, 0.0, 1.0)};
 	}();
 
 	// Draw image
diff --git a/Engine/Plugins/Experimental/ImageWidgets/Source/ImageWidgets/Private/SImageViewport.cpp b/Engine/Plugins/Experimental/ImageWidgets/Source/ImageWidgets/Private/SImageViewport.cpp
--- a/Engine/Plugins/Experimental/ImageWidgets/Source/ImageWidgets/Private/SImageViewport.cpp
+++ b/Engine/Plugins/Experimental/ImageWidgets/Source/ImageWidgets/Private/SImageViewport.cpp
@@ -55,7 +55,8 @@ namespace UE::ImageWidgets
 #if IMAGE_WIDGETS_WITH_AB_COMPARISON
 		: ABComparison(MakePimpl<FImageABComparison>(
 			FImageABComparison::FImageIsValid::CreateLambda([this](const FGuid& Guid) { return ImageViewer->IsValidImage(Guid); }),
-			FImageABComparison::FGetCurrentImageGuid::CreateLambda([this] { return ImageViewer->GetCurrentImageInfo().Guid; })))
+			FImageABComparison::FGetCurrentImageGuid::CreateLambda([this] { return ImageViewer->GetCurrentImageInfo().Guid; }),
+			FImageABComparison::FGetImageName::CreateLambda([this](const FGuid& Guid) { return ImageViewer->GetImageName(Guid); })))
 #endif
 	{
 	}
diff --git a/Engine/Plugins/Experimental/ImageWidgets/Source/ImageWidgets/Private/SImageViewportToolbar.cpp b/Engine/Plugins/Experimental/ImageWidgets/Source/ImageWidgets/Private/SImageViewportToolbar.cpp
--- a/Engine/Plugins/Experimental/ImageWidgets/Source/ImageWidgets/Private/SImageViewportToolbar.cpp
+++ b/Engine/Plugins/Experimental/ImageWidgets/Source/ImageWidgets/Private/SImageViewportToolbar.cpp
@@ -137,18 +137,22 @@ MakeCenterToolbar(Parameters.ToolbarExtender, ABComparison != nullptr)
 					return SNew(SCheckBox)
 					.Style(ButtonStyle)
 					.IsEnabled_Lambda([this, AorB]
-					                      {
-						                      return ABComparison->CanSetABComparison(AorB);
-					                      })
+						{
+							return ABComparison->CanSetABComparison(AorB);
+						})
 					.IsChecked(ABComparison->ABComparisonIsSet(AorB))
 					.OnCheckStateChanged_Lambda([this, AorB](const ECheckBoxState State)
-					                      {
-						                      ABComparison->SetABComparison(AorB, State != ECheckBoxState::Checked ? FGuid() : ImageGuid.Execute());
-					                      })
+					    {
+							ABComparison->SetABComparison(AorB, State != ECheckBoxState::Checked ? FGuid() : ImageGuid.Execute());
+					    })
 					[
 						SNew(STextBlock)
 							.Font(FAppStyle::GetFontStyle("EditorViewportToolBar.Font"))
 							.Text(FText::FromString(Label))
+							.ToolTipText_Lambda([&ABComparison = ABComparison, AorB]
+								{
+									return ABComparison->ABComparisonIsSet(AorB) ? ABComparison->GetName(AorB) : FText();
+								})
 							.Margin(FMargin(2.0f, 0.0f))
 					];
 				};
@@ -160,12 +164,12 @@ MakeCenterToolbar(Parameters.ToolbarExtender, ABComparison != nullptr)
 					+ SHorizontalBox::Slot()
 					.AutoWidth()
 					[
-						GetTextButton("A", ButtonStyleStart, FImageABComparison::A)
+						GetTextButton("A", ButtonStyleStart, FImageABComparison::EAorB::A)
 					]
 					+ SHorizontalBox::Slot()
 					.AutoWidth()
 					[
-						GetTextButton("B", ButtonStyleEnd, FImageABComparison::B)
+						GetTextButton("B", ButtonStyleEnd, FImageABComparison::EAorB::B)
 					];
 
 				ToolbarBuilder.AddToolBarWidget(RGBA);
diff --git a/Engine/Plugins/Experimental/ImageWidgets/Source/ImageWidgets/Public/IImageViewer.h b/Engine/Plugins/Experimental/ImageWidgets/Source/ImageWidgets/Public/IImageViewer.h
--- a/Engine/Plugins/Experimental/ImageWidgets/Source/ImageWidgets/Public/IImageViewer.h
+++ b/Engine/Plugins/Experimental/ImageWidgets/Source/ImageWidgets/Public/IImageViewer.h
@@ -138,6 +138,13 @@ namespace UE::ImageWidgets
 		 * @return true if the GUID represents a currently available image
 		 */
 		virtual bool IsValidImage(const FGuid& Guid) const = 0;
+
+		/**
+		 * Returns the name of a currently available image. This is potentially used in the UI to refer to an image.
+		 * @param Guid Unique identifier of the potentially available image
+		 * @return Name of the image if the GUID represents a currently available image
+		 */
+		virtual FText GetImageName(const FGuid& Guid) const = 0;
 #endif
 
 	protected:
diff --git a/Engine/Plugins/Experimental/ImageWidgets/Source/ImageWidgets/Public/SImageCatalog.h b/Engine/Plugins/Experimental/ImageWidgets/Source/ImageWidgets/Public/SImageCatalog.h
--- a/Engine/Plugins/Experimental/ImageWidgets/Source/ImageWidgets/Public/SImageCatalog.h
+++ b/Engine/Plugins/Experimental/ImageWidgets/Source/ImageWidgets/Public/SImageCatalog.h
@@ -70,12 +70,12 @@ namespace UE::ImageWidgets
 
 		/**
 		 * Adds a regular item to the catalog.
-		 * @param Item Data for the item that is being added. 
+		 * @param Item Data for the item that is being added.
 		 */
 		IMAGEWIDGETS_API void AddItem(const TSharedPtr<FImageCatalogItemData>& Item);
 
 		/**
-		 * Adds a pinned item to the catalog. Pinned items appear in a seperate list above regular items.
+		 * Adds a pinned item to the catalog. Pinned items appear in a separate list above regular items.
 		 * @param Item Data for the item that is being added. 
 		 */
 		IMAGEWIDGETS_API void AddPinnedItem(const TSharedPtr<FImageCatalogItemData>& Item);
diff --git a/Engine/Plugins/Experimental/TextureGraph/Source/TextureGraphEditor/Private/STG_NodePreview.h b/Engine/Plugins/Experimental/TextureGraph/Source/TextureGraphEditor/Private/STG_NodePreview.h
--- a/Engine/Plugins/Experimental/TextureGraph/Source/TextureGraphEditor/Private/STG_NodePreview.h
+++ b/Engine/Plugins/Experimental/TextureGraph/Source/TextureGraphEditor/Private/STG_NodePreview.h
@@ -32,6 +32,7 @@ public:
 	virtual void OnImageSelected(const FGuid& ImageGuid) override {};
 #if IMAGE_WIDGETS_WITH_AB_COMPARISON
 	virtual bool IsValidImage(const FGuid& Guid) const override { return true; }
+	virtual FText GetImageName(const FGuid& Guid) const override { return {}; }
 #endif
 	// IImageViewer overrides - end
 
