ImageWidgets: Remove `EnableAbComparison` build setting and respective code paths

diff --git a/Engine/Plugins/Experimental/ImageWidgets/Source/ImageWidgets/ImageWidgets.Build.cs b/Engine/Plugins/Experimental/ImageWidgets/Source/ImageWidgets/ImageWidgets.Build.cs
--- a/Engine/Plugins/Experimental/ImageWidgets/Source/ImageWidgets/ImageWidgets.Build.cs
+++ b/Engine/Plugins/Experimental/ImageWidgets/Source/ImageWidgets/ImageWidgets.Build.cs
@@ -7,9 +7,6 @@ public class ImageWidgets : ModuleRules
 	// Set this flag to true to force building the ColorViewer sample. Otherwise, it will only be built in Debug builds. 
 	private const bool ForceBuildColorViewerSample = false;
 
-	// Todo Remove build flags as soon as all the respective prototype code is migrated to this module.
-	private const bool EnableAbComparison = true;
-
 	public ImageWidgets(ReadOnlyTargetRules Target) : base(Target)
 	{
 		PCHUsage = PCHUsageMode.UseExplicitOrSharedPCHs;
@@ -46,7 +43,5 @@ public class ImageWidgets : ModuleRules
 			);
 		}
 		PublicDefinitions.Add(bBuildColorViewerSample ? "IMAGE_WIDGETS_BUILD_COLOR_VIEWER_SAMPLE=1" : "IMAGE_WIDGETS_BUILD_COLOR_VIEWER_SAMPLE=0");
-
-		PublicDefinitions.Add(EnableAbComparison ? "IMAGE_WIDGETS_WITH_AB_COMPARISON=1" : "IMAGE_WIDGETS_WITH_AB_COMPARISON=0");
 	}
 }
\ No newline at end of file
diff --git a/Engine/Plugins/Experimental/ImageWidgets/Source/ImageWidgets/Private/ColorViewerSample/ColorViewer.cpp b/Engine/Plugins/Experimental/ImageWidgets/Source/ImageWidgets/Private/ColorViewerSample/ColorViewer.cpp
--- a/Engine/Plugins/Experimental/ImageWidgets/Source/ImageWidgets/Private/ColorViewerSample/ColorViewer.cpp
+++ b/Engine/Plugins/Experimental/ImageWidgets/Source/ImageWidgets/Private/ColorViewerSample/ColorViewer.cpp
@@ -22,7 +22,6 @@ namespace UE::ImageWidgets::Sample
 
 	void FColorViewer::DrawCurrentImage(FViewport* Viewport, FCanvas* Canvas, const FDrawProperties& Properties)
 	{
-#if IMAGE_WIDGETS_WITH_AB_COMPARISON
 		if (Properties.ABComparison.IsActive())
 		{
 			DrawImage(Properties.ABComparison.GuidA.B, Canvas, Properties.Placement, {0.0, 0.0}, {Properties.ABComparison.Threshold, 1.0});
@@ -30,11 +29,8 @@ namespace UE::ImageWidgets::Sample
 		}
 		else
 		{
-#endif
 			DrawImage(SelectedColorIndex, Canvas, Properties.Placement, {0.0, 0.0}, {1.0, 1.0});
-#if IMAGE_WIDGETS_WITH_AB_COMPARISON
 		}
-#endif
 	}
 
 	TOptional<TVariant<FColor, FLinearColor>> FColorViewer::GetCurrentImagePixelColor(FIntPoint PixelCoords, int32 MipLevel) const
@@ -56,7 +52,6 @@ namespace UE::ImageWidgets::Sample
 		}
 	}
 
-#if IMAGE_WIDGETS_WITH_AB_COMPARISON
 	bool FColorViewer::IsValidImage(const FGuid& Guid) const
 	{
 		return ColorIsValid(Guid.B) && Colors[Guid.B].Guid == Guid;
@@ -72,7 +67,6 @@ namespace UE::ImageWidgets::Sample
 		}
 		return {};
 	}
-#endif
 
 	const FColorViewer::FColorItem* FColorViewer::AddColor()
 	{
diff --git a/Engine/Plugins/Experimental/ImageWidgets/Source/ImageWidgets/Private/ColorViewerSample/ColorViewer.h b/Engine/Plugins/Experimental/ImageWidgets/Source/ImageWidgets/Private/ColorViewerSample/ColorViewer.h
--- a/Engine/Plugins/Experimental/ImageWidgets/Source/ImageWidgets/Private/ColorViewerSample/ColorViewer.h
+++ b/Engine/Plugins/Experimental/ImageWidgets/Source/ImageWidgets/Private/ColorViewerSample/ColorViewer.h
@@ -61,10 +61,8 @@ namespace UE::ImageWidgets::Sample
 		virtual void DrawCurrentImage(FViewport* Viewport, FCanvas* Canvas, const FDrawProperties& Properties) override;
 		virtual TOptional<TVariant<FColor, FLinearColor>> GetCurrentImagePixelColor(FIntPoint PixelCoords, int32 MipLevel) const override;
 		virtual void OnImageSelected(const FGuid& Guid) override;
-#if IMAGE_WIDGETS_WITH_AB_COMPARISON
 		virtual bool IsValidImage(const FGuid& Guid) const override;
 		virtual FText GetImageName(const FGuid& Guid) const override;
-#endif
 		// IImageViewer overrides - end
 
 		/** Adds a color item. */
diff --git a/Engine/Plugins/Experimental/ImageWidgets/Source/ImageWidgets/Private/ColorViewerSample/ColorViewerWidget.cpp b/Engine/Plugins/Experimental/ImageWidgets/Source/ImageWidgets/Private/ColorViewerSample/ColorViewerWidget.cpp
--- a/Engine/Plugins/Experimental/ImageWidgets/Source/ImageWidgets/Private/ColorViewerSample/ColorViewerWidget.cpp
+++ b/Engine/Plugins/Experimental/ImageWidgets/Source/ImageWidgets/Private/ColorViewerSample/ColorViewerWidget.cpp
@@ -51,9 +51,7 @@ namespace UE::ImageWidgets::Sample
 									.bBackgroundColorEnabled = false,
 									.bBackgroundCheckerEnabled = false
 								})
-#if IMAGE_WIDGETS_WITH_AB_COMPARISON
 								.bABComparisonEnabled(true)
-#endif
 						]
 			];
 	}
@@ -100,12 +98,7 @@ namespace UE::ImageWidgets::Sample
 
 	TTuple<FText, FText, FText> GetColorItemMetaData(const FColorViewer& ColorViewer, const FColorViewer::FColorItem* ColorItem)
 	{
-#if IMAGE_WIDGETS_WITH_AB_COMPARISON
 		const FText Name = ColorViewer.GetImageName(ColorItem->Guid);
-#else
-		const FString HexColor = FString::Printf(TEXT("#%02X%02X%02X"), ColorItem->Color.R, ColorItem->Color.G, ColorItem->Color.B);
-		const FText Name = FText::FromString(HexColor);
-#endif
 
 		const FText Info = FText::Format(
 			LOCTEXT("ColorEntryInfoLabel", "{0}"), FText::AsTime(ColorItem->DateTime, EDateTimeStyle::Short, FText::GetInvariantTimeZone()));
diff --git a/Engine/Plugins/Experimental/ImageWidgets/Source/ImageWidgets/Private/ImageABComparison.h b/Engine/Plugins/Experimental/ImageWidgets/Source/ImageWidgets/Private/ImageABComparison.h
--- a/Engine/Plugins/Experimental/ImageWidgets/Source/ImageWidgets/Private/ImageABComparison.h
+++ b/Engine/Plugins/Experimental/ImageWidgets/Source/ImageWidgets/Private/ImageABComparison.h
@@ -2,8 +2,6 @@
 
 #pragma once
 
-#if IMAGE_WIDGETS_WITH_AB_COMPARISON
-
 #include "Delegates/Delegate.h"
 #include "Internationalization/Text.h"
 #include "Misc/Guid.h"
@@ -83,6 +81,3 @@ namespace UE::ImageWidgets
 		FGetImageName GetImageName;
 	};
 }
-#else
-#pragma message("Do not include this header!")
-#endif
diff --git a/Engine/Plugins/Experimental/ImageWidgets/Source/ImageWidgets/Private/ImageViewportClient.cpp b/Engine/Plugins/Experimental/ImageWidgets/Source/ImageWidgets/Private/ImageViewportClient.cpp
--- a/Engine/Plugins/Experimental/ImageWidgets/Source/ImageWidgets/Private/ImageViewportClient.cpp
+++ b/Engine/Plugins/Experimental/ImageWidgets/Source/ImageWidgets/Private/ImageViewportClient.cpp
@@ -4,9 +4,7 @@
 
 #include "CanvasItem.h"
 #include "CanvasTypes.h"
-#if IMAGE_WIDGETS_WITH_AB_COMPARISON
 #include "ImageABComparison.h"
-#endif
 #include "ImageWidgetsLogCategory.h"
 #include "SImageViewport.h"
 #include "Texture2DPreview.h"
@@ -103,11 +101,8 @@ void DestroyCheckerTexture(TStrongObjectPtr<UTexture2D>& CheckerTexture)
 FImageViewportClient::FImageViewportClient(const TWeakPtr<SEditorViewport>& InEditorViewport, FGetImageSize&& InGetImageSize, FDrawImage&& InDrawImage,
                                            FGetDrawSettings&& InGetDrawSettings, FGetDPIScaleFactor&& InGetDPIScaleFactor,
                                            FOnLeftMouseButtonPressed&& InOnLeftMouseButtonPressed, FOnLeftMouseButtonReleased&& InOnLeftMouseButtonReleased,
-#if IMAGE_WIDGETS_WITH_AB_COMPARISON       
-                                           const FImageABComparison* ABComparison,
-#endif
-                                           SImageViewport::FControllerSettings::EDefaultZoomMode DefaultZoomMode,
-                                           EMouseCaptureMode InMouseCaptureMode /*= EMouseCaptureMode::CapturePermanently*/)
+                                           const FImageABComparison* ABComparison, SImageViewport::FControllerSettings::EDefaultZoomMode DefaultZoomMode,
+                                           EMouseCaptureMode InMouseCaptureMode)
 	: FEditorViewportClient(nullptr, nullptr, InEditorViewport)
 	, GetImageSize(MoveTemp(InGetImageSize))
 	, DrawImage(MoveTemp(InDrawImage))
@@ -115,9 +110,7 @@ FImageViewportClient::FImageViewportClient(const TWeakPtr<SEditorViewport>& InEd
 	, GetDPIScaleFactor(MoveTemp(InGetDPIScaleFactor))
 	, OnLeftMouseButtonPressed(MoveTemp(InOnLeftMouseButtonPressed))
 	, OnLeftMouseButtonReleased(MoveTemp(InOnLeftMouseButtonReleased))
-#if IMAGE_WIDGETS_WITH_AB_COMPARISON
 	, ABComparison(ABComparison)
-#endif
 	, Controller(static_cast<FImageViewportController::EZoomMode>(DefaultZoomMode))
 	, MouseCaptureMode(InMouseCaptureMode)
 {
@@ -191,7 +184,6 @@ void FImageViewportClient::Draw(FViewport* InViewport, FCanvas* Canvas)
 		Canvas->DrawItem(Background);
 	}
 
-#if IMAGE_WIDGETS_WITH_AB_COMPARISON
 	const double ABComparisonDividerX = [this]
 	{
 		if (bDraggingABComparisonDivider)
@@ -225,10 +217,6 @@ void FImageViewportClient::Draw(FViewport* InViewport, FCanvas* Canvas)
 		Line.SetColor({0.5f, 0.5f, 0.5f, 1.0f});
 		Canvas->DrawItem(Line);
 	}
-#else
-	// Draw image
-	DrawImage.Execute(InViewport, Canvas, {CachedPlacement, MipProperties});
-#endif
 }
 
 EMouseCursor::Type FImageViewportClient::GetCursor(FViewport* InViewport, int32 X, int32 Y)	
@@ -240,12 +228,10 @@ EMouseCursor::Type FImageViewportClient::GetCursor(FViewport* InViewport, int32
 		return EMouseCursor::GrabHandClosed;
 	}
 
-#if IMAGE_WIDGETS_WITH_AB_COMPARISON
 	if (bDraggingABComparisonDivider || MouseIsOverABComparisonDivider({X, Y}))
 	{
 		return EMouseCursor::ResizeLeftRight;
 	}
-#endif
 
 	const auto [PixelCoordsValid, PixelCoords] = GetPixelCoordinatesUnderCursor();
 	const FIntPoint ImageSize = GetImageSize.Execute();
@@ -311,13 +297,11 @@ void FImageViewportClient::TrackingStarted(const FInputEventState& InputState, b
 		FIntPoint MousePos;
 		InputState.GetViewport()->GetMousePos(MousePos);
 
-#if IMAGE_WIDGETS_WITH_AB_COMPARISON
 		if (InputState.IsLeftMouseButtonPressed() && MouseIsOverABComparisonDivider(MousePos))
 		{
 			bDraggingABComparisonDivider = true;
 			return;
 		}
-#endif
 
 		if (InputState.IsRightMouseButtonPressed())
 		{
@@ -332,7 +316,6 @@ void FImageViewportClient::TrackingStarted(const FInputEventState& InputState, b
 
 void FImageViewportClient::TrackingStopped()
 {
-#if IMAGE_WIDGETS_WITH_AB_COMPARISON
 	if (bDraggingABComparisonDivider)
 	{
 		bDraggingABComparisonDivider = false;
@@ -347,7 +330,6 @@ void FImageViewportClient::TrackingStopped()
 		
 		ABComparisonDivider = DraggingEndPos.X / ImageSize.X;
 	}
-#endif
 
 	if (bDragging)
 	{
@@ -489,13 +471,11 @@ EMouseCaptureMode FImageViewportClient::GetMouseCaptureMode() const
 	return MouseCaptureMode;
 }
 
-#if IMAGE_WIDGETS_WITH_AB_COMPARISON
 bool FImageViewportClient::MouseIsOverABComparisonDivider(const FIntPoint MousePos) const
 {
 	const double ABComparisonDividerPosition = CachedPlacement.Offset.X + CachedPlacement.Size.X * ABComparisonDivider;
 	return ABComparisonDividerPosition - 1.0 <= MousePos.X && MousePos.X <= ABComparisonDividerPosition + 1.0;
 }
-#endif
 }
 
 #undef LOCTEXT_NAMESPACE
diff --git a/Engine/Plugins/Experimental/ImageWidgets/Source/ImageWidgets/Private/ImageViewportClient.h b/Engine/Plugins/Experimental/ImageWidgets/Source/ImageWidgets/Private/ImageViewportClient.h
--- a/Engine/Plugins/Experimental/ImageWidgets/Source/ImageWidgets/Private/ImageViewportClient.h
+++ b/Engine/Plugins/Experimental/ImageWidgets/Source/ImageWidgets/Private/ImageViewportClient.h
@@ -27,10 +27,8 @@ namespace UE::ImageWidgets
 		FImageViewportClient(const TWeakPtr<SEditorViewport>& InViewport, FGetImageSize&& InGetImageSize, FDrawImage&& InDrawImage,
 		                     FGetDrawSettings&& InGetDrawSettings, FGetDPIScaleFactor&& InGetDPIScaleFactor,
 		                     FOnLeftMouseButtonPressed&& InOnLeftMouseButtonPressed, FOnLeftMouseButtonReleased&& InOnLeftMouseButtonReleased,
-#if IMAGE_WIDGETS_WITH_AB_COMPARISON
-							 const FImageABComparison* ABComparison,
-#endif
-		                     SImageViewport::FControllerSettings::EDefaultZoomMode DefaultZoomMode, EMouseCaptureMode InMouseCaptureMode = EMouseCaptureMode::CapturePermanently);
+		                     const FImageABComparison* ABComparison, SImageViewport::FControllerSettings::EDefaultZoomMode DefaultZoomMode,
+		                     EMouseCaptureMode InMouseCaptureMode);
 		virtual ~FImageViewportClient() override;
 
 		virtual void Draw(FViewport* InViewport, FCanvas* Canvas) override;
@@ -68,9 +66,7 @@ namespace UE::ImageWidgets
 		IImageViewer::FDrawProperties::FMip GetMipProperties() const;
 
 		void CreateOrDestroyCheckerTextureIfSettingsChanged(const SImageViewport::FDrawSettings& DrawSettings);
-#if IMAGE_WIDGETS_WITH_AB_COMPARISON
 		bool MouseIsOverABComparisonDivider(FIntPoint MousePos) const;
-#endif
 
 		FVector2d GetViewportSizeWithDPIScaling() const;
 		
@@ -80,9 +76,7 @@ namespace UE::ImageWidgets
 		FGetDPIScaleFactor GetDPIScaleFactor;
 		FOnLeftMouseButtonPressed OnLeftMouseButtonPressed;
 		FOnLeftMouseButtonReleased OnLeftMouseButtonReleased;
-#if IMAGE_WIDGETS_WITH_AB_COMPARISON
 		const FImageABComparison* ABComparison;
-#endif
 		
 		bool bDragging = false;
 		FIntPoint DraggingStart;
@@ -98,10 +92,8 @@ namespace UE::ImageWidgets
 		TStrongObjectPtr<UTexture2D> CheckerTexture;
 		FCheckerTextureSettings CachedCheckerTextureSettings;
 
-#if IMAGE_WIDGETS_WITH_AB_COMPARISON
 		double ABComparisonDivider = 0.5;
 		bool bDraggingABComparisonDivider = false;
-#endif
 		
 		EMouseCaptureMode MouseCaptureMode;
 	};
diff --git a/Engine/Plugins/Experimental/ImageWidgets/Source/ImageWidgets/Private/ImageWidgetsModule.cpp b/Engine/Plugins/Experimental/ImageWidgets/Source/ImageWidgets/Private/ImageWidgetsModule.cpp
--- a/Engine/Plugins/Experimental/ImageWidgets/Source/ImageWidgets/Private/ImageWidgetsModule.cpp
+++ b/Engine/Plugins/Experimental/ImageWidgets/Source/ImageWidgets/Private/ImageWidgetsModule.cpp
@@ -37,7 +37,7 @@ namespace UE::ImageWidgets
 
 	private:
 		TSharedRef<SDockTab> MakeColorViewerTab(const FSpawnTabArgs&);
-		TSharedRef<SWidget> GetColorViewer(const TSharedRef<SDockTab>& InParentTab);
+		TSharedRef<SWidget> GetColorViewer();
 
 		bool bHasRegisteredTabSpawners = false;
 		TWeakPtr<Sample::SColorViewerWidget> ColorViewerPtr;
@@ -91,11 +91,11 @@ namespace UE::ImageWidgets
 	{
 		TSharedRef<SDockTab> WidgetReflectorTab = SNew(SDockTab)
 			.TabRole(NomadTab);
-		WidgetReflectorTab->SetContent(GetColorViewer(WidgetReflectorTab));
+		WidgetReflectorTab->SetContent(GetColorViewer());
 		return WidgetReflectorTab;
 	}
 
-	TSharedRef<SWidget> FImageWidgetsModule::GetColorViewer(const TSharedRef<SDockTab>& InParentTab)
+	TSharedRef<SWidget> FImageWidgetsModule::GetColorViewer()
 	{
 		TSharedPtr<Sample::SColorViewerWidget> ColorViewer = ColorViewerPtr.Pin();
 
diff --git a/Engine/Plugins/Experimental/ImageWidgets/Source/ImageWidgets/Private/SImageViewport.cpp b/Engine/Plugins/Experimental/ImageWidgets/Source/ImageWidgets/Private/SImageViewport.cpp
--- a/Engine/Plugins/Experimental/ImageWidgets/Source/ImageWidgets/Private/SImageViewport.cpp
+++ b/Engine/Plugins/Experimental/ImageWidgets/Source/ImageWidgets/Private/SImageViewport.cpp
@@ -3,9 +3,7 @@
 #include "SImageViewport.h"
 
 #include "CanvasTypes.h"
-#if IMAGE_WIDGETS_WITH_AB_COMPARISON
 #include "ImageABComparison.h"
-#endif
 #include "ImageViewportClient.h"
 #include "ImageViewportController.h"
 #include "ImageWidgetsCommands.h"
@@ -52,12 +50,10 @@ namespace UE::ImageWidgets
 	}
 
 	SImageViewport::SImageViewport()
-#if IMAGE_WIDGETS_WITH_AB_COMPARISON
 		: ABComparison(MakePimpl<FImageABComparison>(
 			FImageABComparison::FImageIsValid::CreateLambda([this](const FGuid& Guid) { return ImageViewer->IsValidImage(Guid); }),
 			FImageABComparison::FGetCurrentImageGuid::CreateLambda([this] { return ImageViewer->GetCurrentImageInfo().Guid; }),
 			FImageABComparison::FGetImageName::CreateLambda([this](const FGuid& Guid) { return ImageViewer->GetImageName(Guid); })))
-#endif
 	{
 	}
 
@@ -157,9 +153,7 @@ namespace UE::ImageWidgets
 		StatusBarExtender = InArgs._StatusBarExtender;
 
 		DrawSettings = InArgs._DrawSettings;
-#if IMAGE_WIDGETS_WITH_AB_COMPARISON
 		bABComparisonEnabled = InArgs._bABComparisonEnabled;
-#endif
 		OnLeftMouseButtonPressed = InArgs._OnLeftMouseButtonPressed;
 		OnLeftMouseButtonReleased = InArgs._OnLeftMouseButtonReleased;
 
@@ -210,9 +204,7 @@ namespace UE::ImageWidgets
 			FGetDPIScaleFactor::CreateLambda(GetDPIScaleFactor),
 			FOnLeftMouseButtonPressed::CreateLambda(HandleOnLeftMouseButtonPressed),
 			FOnLeftMouseButtonReleased::CreateLambda(HandleOnLeftMouseButtonReleased),
-#if IMAGE_WIDGETS_WITH_AB_COMPARISON
 			ABComparison.Get(),
-#endif
 			InArgs._ControllerSettings.DefaultZoomMode,
 			InArgs._MouseCaptureMode)
 			);
@@ -301,7 +293,6 @@ namespace UE::ImageWidgets
 			return ImageViewer->GetCurrentImageInfo().Guid;
 		};
 
-#if IMAGE_WIDGETS_WITH_AB_COMPARISON
 		SAssignNew(ImageViewportToolbar, SImageViewportToolbar, StaticCastSharedPtr<FImageViewportClient>(Client), CommandList,
 						  SImageViewportToolbar::FConstructParameters
 						  {
@@ -311,16 +302,6 @@ namespace UE::ImageWidgets
 							  bABComparisonEnabled ? ABComparison.Get() : nullptr,
 							  ToolbarExtender
 						  });
-#else
-		SAssignNew(ImageViewportToolbar, SImageViewportToolbar, StaticCastSharedPtr<FImageViewportClient>(Client), CommandList,
-						  SImageViewportToolbar::FConstructParameters
-						  {
-							  SImageViewportToolbar::FHasImage::CreateLambda(HasImage),
-							  SImageViewportToolbar::FNumMips::CreateLambda(GetNumMips),
-							  SImageViewportToolbar::FImageGuid::CreateLambda(GetImageGuid),
-							  ToolbarExtender
-						  });
-#endif
 						  
 
 		ToolbarExtender.Reset();
diff --git a/Engine/Plugins/Experimental/ImageWidgets/Source/ImageWidgets/Private/SImageViewportToolbar.cpp b/Engine/Plugins/Experimental/ImageWidgets/Source/ImageWidgets/Private/SImageViewportToolbar.cpp
--- a/Engine/Plugins/Experimental/ImageWidgets/Source/ImageWidgets/Private/SImageViewportToolbar.cpp
+++ b/Engine/Plugins/Experimental/ImageWidgets/Source/ImageWidgets/Private/SImageViewportToolbar.cpp
@@ -2,9 +2,7 @@
 
 #include "SImageViewportToolbar.h"
 
-#if IMAGE_WIDGETS_WITH_AB_COMPARISON
 #include "ImageABComparison.h"
-#endif
 #include "ImageViewportClient.h"
 #include "ImageWidgetsCommands.h"
 #include "SEditorViewportToolBarMenu.h"
@@ -38,9 +36,7 @@ namespace UE::ImageWidgets
 		check(NumMips.IsBound());
 		check(ImageGuid.IsBound());
 
-#if IMAGE_WIDGETS_WITH_AB_COMPARISON
 		ABComparison = Parameters.ABComparison;
-#endif
 
 		ChildSlot
 		[
@@ -58,11 +54,7 @@ namespace UE::ImageWidgets
 				  .FillWidth(1.0f)
 				  .HAlign(HAlign_Center)
 				[
-#if IMAGE_WIDGETS_WITH_AB_COMPARISON
-MakeCenterToolbar(Parameters.ToolbarExtender, ABComparison != nullptr)
-#else
-					MakeCenterToolbar(Parameters.ToolbarExtender)
-#endif
+					MakeCenterToolbar(Parameters.ToolbarExtender, ABComparison != nullptr)
 				]
 				+ SHorizontalBox::Slot()
 				  .AutoWidth()
@@ -121,7 +113,6 @@ MakeCenterToolbar(Parameters.ToolbarExtender, ABComparison != nullptr)
 		return ToolbarBuilder.MakeWidget();
 	}
 
-	#if IMAGE_WIDGETS_WITH_AB_COMPARISON
 	TSharedRef<SWidget> SImageViewportToolbar::MakeCenterToolbar(const TSharedPtr<FExtender>& Extender, bool bEnableABComparison)
 	{
 		FSlimHorizontalToolBarBuilder ToolbarBuilder = GetToolbarBuilder(CommandList, Extender);
@@ -183,23 +174,6 @@ MakeCenterToolbar(Parameters.ToolbarExtender, ABComparison != nullptr)
 
 		return ToolbarBuilder.MakeWidget();
 	}
-#else
-	TSharedRef<SWidget> SImageViewportToolbar::MakeCenterToolbar(const TSharedPtr<FExtender>& Extender)
-	{
-		FSlimHorizontalToolBarBuilder ToolbarBuilder = GetToolbarBuilder(CommandList, Extender);
-
-		ToolbarBuilder.BeginSection("ToolbarCenter");
-		{
-			// This is deliberately left empty.
-			// Toolbar extenders use this section to add additional widgets.
-		}
-		ToolbarBuilder.EndSection();
-
-		ToolbarBuilder.AddSeparator();
-
-		return ToolbarBuilder.MakeWidget();
-	}
-#endif
 
 	TSharedRef<SWidget> SImageViewportToolbar::MakeRightToolbar(const TSharedPtr<FExtender>& Extender)
 	{
diff --git a/Engine/Plugins/Experimental/ImageWidgets/Source/ImageWidgets/Private/SImageViewportToolbar.h b/Engine/Plugins/Experimental/ImageWidgets/Source/ImageWidgets/Private/SImageViewportToolbar.h
--- a/Engine/Plugins/Experimental/ImageWidgets/Source/ImageWidgets/Private/SImageViewportToolbar.h
+++ b/Engine/Plugins/Experimental/ImageWidgets/Source/ImageWidgets/Private/SImageViewportToolbar.h
@@ -13,9 +13,7 @@ class SViewportToolBarComboMenu;
 
 namespace UE::ImageWidgets
 {
-#if IMAGE_WIDGETS_WITH_AB_COMPARISON
 	class FImageABComparison;
-#endif
 	class FImageViewportClient;
 
 	/**
@@ -39,9 +37,7 @@ namespace UE::ImageWidgets
 			FHasImage HasImage;
 			FNumMips NumMips;
 			FImageGuid ImageGuid;
-#if IMAGE_WIDGETS_WITH_AB_COMPARISON
 			FImageABComparison* ABComparison = nullptr;
-#endif
 			TSharedPtr<FExtender> ToolbarExtender;
 		};
 		
@@ -50,11 +46,7 @@ namespace UE::ImageWidgets
 
 	private:
 		TSharedRef<SWidget> MakeLeftToolbar(const TSharedPtr<FExtender>& Extender);
-#if IMAGE_WIDGETS_WITH_AB_COMPARISON
 		TSharedRef<SWidget> MakeCenterToolbar(const TSharedPtr<FExtender>& Extender, bool bEnableABComparison);
-#else
-		TSharedRef<SWidget> MakeCenterToolbar(const TSharedPtr<FExtender>& Extender);
-#endif
 		TSharedRef<SWidget> MakeRightToolbar(const TSharedPtr<FExtender>& Extender);
 
 		EVisibility GetMipMenuVisibility() const;
@@ -70,8 +62,6 @@ namespace UE::ImageWidgets
 		FHasImage HasImage;
 		FNumMips NumMips;
 		FImageGuid ImageGuid;
-#if IMAGE_WIDGETS_WITH_AB_COMPARISON
 		FImageABComparison* ABComparison = nullptr;
-#endif
 	};
 }
diff --git a/Engine/Plugins/Experimental/ImageWidgets/Source/ImageWidgets/Public/IImageViewer.h b/Engine/Plugins/Experimental/ImageWidgets/Source/ImageWidgets/Public/IImageViewer.h
--- a/Engine/Plugins/Experimental/ImageWidgets/Source/ImageWidgets/Public/IImageViewer.h
+++ b/Engine/Plugins/Experimental/ImageWidgets/Source/ImageWidgets/Public/IImageViewer.h
@@ -68,7 +68,6 @@ namespace UE::ImageWidgets
 				float MipLevel;
 			};
 
-#if IMAGE_WIDGETS_WITH_AB_COMPARISON
 			/**
 			 * Information necessary for rendering AB comparisons.
 			 */
@@ -91,13 +90,10 @@ namespace UE::ImageWidgets
 				 */
 				double Threshold;
 			};
-#endif
 
 			FPlacement Placement;
 			FMip Mip;
-#if IMAGE_WIDGETS_WITH_AB_COMPARISON
 			FABComparison ABComparison;
-#endif
 		};
 
 		/**
@@ -131,7 +127,6 @@ namespace UE::ImageWidgets
 		 */
 		virtual void OnImageSelected(const FGuid& Guid) = 0;
 
-#if IMAGE_WIDGETS_WITH_AB_COMPARISON
 		/**
 		 * Returns if a given GUID represents a currently available image.
 		 * @param Guid Unique identifier of the potentially available image
@@ -145,7 +140,6 @@ namespace UE::ImageWidgets
 		 * @return Name of the image if the GUID represents a currently available image
 		 */
 		virtual FText GetImageName(const FGuid& Guid) const = 0;
-#endif
 
 	protected:
 		~IImageViewer() = default;
diff --git a/Engine/Plugins/Experimental/ImageWidgets/Source/ImageWidgets/Public/SImageViewport.h b/Engine/Plugins/Experimental/ImageWidgets/Source/ImageWidgets/Public/SImageViewport.h
--- a/Engine/Plugins/Experimental/ImageWidgets/Source/ImageWidgets/Public/SImageViewport.h
+++ b/Engine/Plugins/Experimental/ImageWidgets/Source/ImageWidgets/Public/SImageViewport.h
@@ -10,9 +10,7 @@ class SViewportToolBar;
 
 namespace UE::ImageWidgets
 {
-#if IMAGE_WIDGETS_WITH_AB_COMPARISON
 	class FImageABComparison;
-#endif
 	class FImageViewportClient;
 	class FStatusBarExtension;
 	class IImageViewer;
@@ -118,12 +116,8 @@ namespace UE::ImageWidgets
 		DECLARE_DELEGATE(FOnLeftMouseButtonReleased);
 
 		SLATE_BEGIN_ARGS(SImageViewport)
-#if IMAGE_WIDGETS_WITH_AB_COMPARISON
 			: _bABComparisonEnabled(false)
 			, _MouseCaptureMode(EMouseCaptureMode::CapturePermanently)
-#else
-			: _MouseCaptureMode(EMouseCaptureMode::CapturePermanently)
-#endif
 			{
 			}
 
@@ -136,10 +130,8 @@ namespace UE::ImageWidgets
 			/** Settings for drawing viewport contents other than the actual image */
 			SLATE_ATTRIBUTE(FDrawSettings, DrawSettings)
 
-#if IMAGE_WIDGETS_WITH_AB_COMPARISON
 			/** Enables AB comparison controls in the toolbar */
 			SLATE_ARGUMENT(bool, bABComparisonEnabled)
-#endif
 
 			/** Settings for controlling the viewport */
 			SLATE_ARGUMENT(FControllerSettings, ControllerSettings)
@@ -231,13 +223,11 @@ namespace UE::ImageWidgets
 		/** Makes the draw setting available either as fixed values or via a callback to the outside of the viewport. */
 		TAttribute<FDrawSettings> DrawSettings;
 
-#if IMAGE_WIDGETS_WITH_AB_COMPARISON
 		/** Flag that determines is AB comparison widgets are enabled or not. The value does not change after the call to @see Construct(). */
 		bool bABComparisonEnabled = false;
 
 		/** Data and logic related to AB comparisons. This is effectively unused when @see bABComparisonEnabled is set to false. */
 		TPimplPtr<FImageABComparison> ABComparison;
-#endif
 
 		/** The image viewer that holds and draws the actual images. */
 		TSharedPtr<IImageViewer> ImageViewer;
diff --git a/Engine/Plugins/Experimental/TextureGraph/Source/TextureGraphEditor/Private/STG_NodePreview.h b/Engine/Plugins/Experimental/TextureGraph/Source/TextureGraphEditor/Private/STG_NodePreview.h
--- a/Engine/Plugins/Experimental/TextureGraph/Source/TextureGraphEditor/Private/STG_NodePreview.h
+++ b/Engine/Plugins/Experimental/TextureGraph/Source/TextureGraphEditor/Private/STG_NodePreview.h
@@ -30,10 +30,8 @@ public:
 	virtual void DrawCurrentImage(FViewport* Viewport, FCanvas* Canvas, const FDrawProperties& Properties) override;
 	virtual TOptional<TVariant<FColor, FLinearColor>> GetCurrentImagePixelColor(FIntPoint PixelCoords, int32 MipIndex) const override;
 	virtual void OnImageSelected(const FGuid& ImageGuid) override {};
-#if IMAGE_WIDGETS_WITH_AB_COMPARISON
 	virtual bool IsValidImage(const FGuid& Guid) const override { return true; }
 	virtual FText GetImageName(const FGuid& Guid) const override { return {}; }
-#endif
 	// IImageViewer overrides - end
 
 	/** Returns the format label for the status bar */
