ImageWidgets:
* Refactor image viewport custom input key handling
* Add settings to customize image viewport overlay

diff --git a/Engine/Plugins/Experimental/ImageWidgets/Source/ImageWidgets/Private/ImageViewportClient.cpp b/Engine/Plugins/Experimental/ImageWidgets/Source/ImageWidgets/Private/ImageViewportClient.cpp
--- a/Engine/Plugins/Experimental/ImageWidgets/Source/ImageWidgets/Private/ImageViewportClient.cpp
+++ b/Engine/Plugins/Experimental/ImageWidgets/Source/ImageWidgets/Private/ImageViewportClient.cpp
@@ -100,25 +100,19 @@ void DestroyCheckerTexture(TStrongObjectPtr<UTexture2D>& CheckerTexture)
 
 FImageViewportClient::FImageViewportClient(const TWeakPtr<SEditorViewport>& InEditorViewport, FGetImageSize&& InGetImageSize, FDrawImage&& InDrawImage,
                                            FGetDrawSettings&& InGetDrawSettings, FGetDPIScaleFactor&& InGetDPIScaleFactor,
-                                           FOnLeftMouseButtonPressed&& InOnLeftMouseButtonPressed, FOnLeftMouseButtonReleased&& InOnLeftMouseButtonReleased,
-                                           const FImageABComparison* ABComparison, SImageViewport::FControllerSettings::EDefaultZoomMode DefaultZoomMode,
-                                           EMouseCaptureMode InMouseCaptureMode)
+                                           const FImageABComparison* InABComparison, const SImageViewport::FControllerSettings& InControllerSettings)
 	: FEditorViewportClient(nullptr, nullptr, InEditorViewport)
 	, GetImageSize(MoveTemp(InGetImageSize))
 	, DrawImage(MoveTemp(InDrawImage))
 	, GetDrawSettings(MoveTemp(InGetDrawSettings))
 	, GetDPIScaleFactor(MoveTemp(InGetDPIScaleFactor))
-	, OnLeftMouseButtonPressed(MoveTemp(InOnLeftMouseButtonPressed))
-	, OnLeftMouseButtonReleased(MoveTemp(InOnLeftMouseButtonReleased))
-	, ABComparison(ABComparison)
-	, Controller(static_cast<FImageViewportController::EZoomMode>(DefaultZoomMode))
-	, MouseCaptureMode(InMouseCaptureMode)
+	, OnInputKey(InControllerSettings.OnInputKey)
+	, ABComparison(InABComparison)
+	, Controller(static_cast<FImageViewportController::EZoomMode>(InControllerSettings.DefaultZoomMode))
 {
 	check(GetImageSize.IsBound());
 	check(DrawImage.IsBound());
 	check(GetDrawSettings.IsBound());
-	check(OnLeftMouseButtonPressed.IsBound());
-	check(OnLeftMouseButtonReleased.IsBound());
 
 	SetRealtime(true);
 }
@@ -275,15 +269,11 @@ bool FImageViewportClient::InputKey(const FInputKeyEventArgs& EventArgs)
 		}
 	}
 
-	if (EventArgs.Key == EKeys::LeftMouseButton)
+	if (OnInputKey.IsBound())
 	{
-		if(EventArgs.Event == IE_Pressed)
+		if (OnInputKey.Execute(EventArgs))
 		{
-			OnLeftMouseButtonPressed.Execute();
-		}
-		else if(EventArgs.Event == IE_Released)
-		{
-			OnLeftMouseButtonReleased.Execute();
+			return true;
 		}
 	}
 
@@ -466,11 +456,6 @@ FVector2d FImageViewportClient::GetViewportSizeWithDPIScaling() const
 	return FVector2d(ViewportSize) / DPIScaleFactor;
 }
 
-EMouseCaptureMode FImageViewportClient::GetMouseCaptureMode() const
-{
-	return MouseCaptureMode;
-}
-
 bool FImageViewportClient::MouseIsOverABComparisonDivider(const FIntPoint MousePos) const
 {
 	const double ABComparisonDividerPosition = CachedPlacement.Offset.X + CachedPlacement.Size.X * ABComparisonDivider;
diff --git a/Engine/Plugins/Experimental/ImageWidgets/Source/ImageWidgets/Private/ImageViewportClient.h b/Engine/Plugins/Experimental/ImageWidgets/Source/ImageWidgets/Private/ImageViewportClient.h
--- a/Engine/Plugins/Experimental/ImageWidgets/Source/ImageWidgets/Private/ImageViewportClient.h
+++ b/Engine/Plugins/Experimental/ImageWidgets/Source/ImageWidgets/Private/ImageViewportClient.h
@@ -15,8 +15,6 @@ namespace UE::ImageWidgets
 	DECLARE_DELEGATE_ThreeParams(FDrawImage, FViewport*, FCanvas*, const IImageViewer::FDrawProperties&)
 	DECLARE_DELEGATE_RetVal(SImageViewport::FDrawSettings, FGetDrawSettings)
 	DECLARE_DELEGATE_RetVal(float, FGetDPIScaleFactor)
-	DECLARE_DELEGATE(FOnLeftMouseButtonPressed);
-	DECLARE_DELEGATE(FOnLeftMouseButtonReleased);
 
 	/**
 	 * Viewport client for controlling the camera and drawing viewport contents. 
@@ -24,11 +22,9 @@ namespace UE::ImageWidgets
 	class FImageViewportClient : public FEditorViewportClient
 	{
 	public:
-		FImageViewportClient(const TWeakPtr<SEditorViewport>& InViewport, FGetImageSize&& InGetImageSize, FDrawImage&& InDrawImage,
+		FImageViewportClient(const TWeakPtr<SEditorViewport>& InEditorViewport, FGetImageSize&& InGetImageSize, FDrawImage&& InDrawImage,
 		                     FGetDrawSettings&& InGetDrawSettings, FGetDPIScaleFactor&& InGetDPIScaleFactor,
-		                     FOnLeftMouseButtonPressed&& InOnLeftMouseButtonPressed, FOnLeftMouseButtonReleased&& InOnLeftMouseButtonReleased,
-		                     const FImageABComparison* ABComparison, SImageViewport::FControllerSettings::EDefaultZoomMode DefaultZoomMode,
-		                     EMouseCaptureMode InMouseCaptureMode);
+		                     const FImageABComparison* InABComparison, const SImageViewport::FControllerSettings& InControllerSettings);
 		virtual ~FImageViewportClient() override;
 
 		virtual void Draw(FViewport* InViewport, FCanvas* Canvas) override;
@@ -38,7 +34,6 @@ namespace UE::ImageWidgets
 		virtual bool InputKey(const FInputKeyEventArgs& EventArgs) override;
 		virtual void TrackingStarted(const FInputEventState& InputState, bool bIsDraggingWidget, bool bNudge) override;
 		virtual void TrackingStopped() override;
-		virtual EMouseCaptureMode GetMouseCaptureMode() const override;
 		// SEditorViewport overrides - end
 
 		int32 GetMipLevel() const;
@@ -74,8 +69,7 @@ namespace UE::ImageWidgets
 		FDrawImage DrawImage;
 		FGetDrawSettings GetDrawSettings;
 		FGetDPIScaleFactor GetDPIScaleFactor;
-		FOnLeftMouseButtonPressed OnLeftMouseButtonPressed;
-		FOnLeftMouseButtonReleased OnLeftMouseButtonReleased;
+		SImageViewport::FControllerSettings::FOnInputKey OnInputKey;
 		const FImageABComparison* ABComparison;
 		
 		bool bDragging = false;
@@ -94,7 +88,5 @@ namespace UE::ImageWidgets
 
 		double ABComparisonDivider = 0.5;
 		bool bDraggingABComparisonDivider = false;
-		
-		EMouseCaptureMode MouseCaptureMode;
 	};
 }
diff --git a/Engine/Plugins/Experimental/ImageWidgets/Source/ImageWidgets/Private/SImageViewport.cpp b/Engine/Plugins/Experimental/ImageWidgets/Source/ImageWidgets/Private/SImageViewport.cpp
--- a/Engine/Plugins/Experimental/ImageWidgets/Source/ImageWidgets/Private/SImageViewport.cpp
+++ b/Engine/Plugins/Experimental/ImageWidgets/Source/ImageWidgets/Private/SImageViewport.cpp
@@ -153,9 +153,8 @@ namespace UE::ImageWidgets
 		StatusBarExtender = InArgs._StatusBarExtender;
 
 		DrawSettings = InArgs._DrawSettings;
+		OverlaySettings = InArgs._OverlaySettings;
 		bABComparisonEnabled = InArgs._bABComparisonEnabled;
-		OnLeftMouseButtonPressed = InArgs._OnLeftMouseButtonPressed;
-		OnLeftMouseButtonReleased = InArgs._OnLeftMouseButtonReleased;
 
 		auto GetImageSize = [this]
 		{
@@ -174,7 +173,7 @@ namespace UE::ImageWidgets
 
 		auto GetDrawSettings = [this]
 		{
-			return DrawSettings.Get();
+			return DrawSettings.Get({});
 		};
 
 		auto GetDPIScaleFactor = [this]
@@ -186,27 +185,14 @@ namespace UE::ImageWidgets
 			return 1.0f;
 		};
 
-		auto HandleOnLeftMouseButtonPressed = [this]()
-		{
-			OnLeftMouseButtonPressed.ExecuteIfBound();
-		};
-
-		auto HandleOnLeftMouseButtonReleased = [this]()
-		{
-			OnLeftMouseButtonReleased.ExecuteIfBound();
-		};
-
 		ImageViewportClient = MakeShareable(new FImageViewportClient(
 			StaticCastWeakPtr<SEditorViewport>(AsWeak()),
 			FGetImageSize::CreateLambda(GetImageSize),
 			FDrawImage::CreateLambda(DrawImage),
 			FGetDrawSettings::CreateLambda(GetDrawSettings),
 			FGetDPIScaleFactor::CreateLambda(GetDPIScaleFactor),
-			FOnLeftMouseButtonPressed::CreateLambda(HandleOnLeftMouseButtonPressed),
-			FOnLeftMouseButtonReleased::CreateLambda(HandleOnLeftMouseButtonReleased),
 			ABComparison.Get(),
-			InArgs._ControllerSettings.DefaultZoomMode,
-			InArgs._MouseCaptureMode)
+			InArgs._ControllerSettings)
 			);
 
 		SEditorViewport::Construct(SEditorViewport::FArguments());
@@ -293,12 +279,18 @@ namespace UE::ImageWidgets
 			return ImageViewer->GetCurrentImageInfo().Guid;
 		};
 
+		auto GetOverlaySettings = [this]
+		{
+			return OverlaySettings.Get({});
+		};
+
 		SAssignNew(ImageViewportToolbar, SImageViewportToolbar, StaticCastSharedPtr<FImageViewportClient>(Client), CommandList,
 						  SImageViewportToolbar::FConstructParameters
 						  {
 							  SImageViewportToolbar::FHasImage::CreateLambda(HasImage),
 							  SImageViewportToolbar::FNumMips::CreateLambda(GetNumMips),
 							  SImageViewportToolbar::FImageGuid::CreateLambda(GetImageGuid),
+							  SImageViewportToolbar::FGetOverlaySettings::CreateLambda(GetOverlaySettings),
 							  bABComparisonEnabled ? ABComparison.Get() : nullptr,
 							  ToolbarExtender
 						  });
@@ -426,7 +418,11 @@ namespace UE::ImageWidgets
 			             .HAlign(HAlign_Left)
 			[
 				SNew(STextBlock)
-				.Text(this, &SImageViewport::GetResolutionLabel)
+					.Visibility_Lambda([&OverlaySettings = OverlaySettings]
+						{
+							return OverlaySettings.IsSet() && OverlaySettings.Get().bDisableStatusBarLeft ? EVisibility::Collapsed : EVisibility::Visible;
+						})
+					.Text(this, &SImageViewport::GetResolutionLabel)
 			];
 		}
 		ApplyHook(FName("StatusBarLeft"), EExtensionHook::After);
@@ -452,6 +448,10 @@ namespace UE::ImageWidgets
 			             .HAlign(HAlign_Right)
 			[
 				SNew(SRichTextBlock)
+					.Visibility_Lambda([&OverlaySettings = OverlaySettings]
+						{
+							return OverlaySettings.IsSet() && OverlaySettings.Get().bDisableStatusBarRight ? EVisibility::Collapsed : EVisibility::Visible;
+						})
 					.Text(this, &SImageViewport::GetPickerLabel)
 					.DecoratorStyleSet(&FImageWidgetsStyle::Get())
 			];
diff --git a/Engine/Plugins/Experimental/ImageWidgets/Source/ImageWidgets/Private/SImageViewportToolbar.cpp b/Engine/Plugins/Experimental/ImageWidgets/Source/ImageWidgets/Private/SImageViewportToolbar.cpp
--- a/Engine/Plugins/Experimental/ImageWidgets/Source/ImageWidgets/Private/SImageViewportToolbar.cpp
+++ b/Engine/Plugins/Experimental/ImageWidgets/Source/ImageWidgets/Private/SImageViewportToolbar.cpp
@@ -32,9 +32,11 @@ namespace UE::ImageWidgets
 		HasImage = MoveTemp(Parameters.HasImage);
 		NumMips = MoveTemp(Parameters.NumMips);
 		ImageGuid = MoveTemp(Parameters.ImageGuid);
+		GetOverlaySettings = MoveTemp(Parameters.GetOverlaySettings);
 		check(HasImage.IsBound());
 		check(NumMips.IsBound());
 		check(ImageGuid.IsBound());
+		check(GetOverlaySettings.IsBound());
 
 		ABComparison = Parameters.ABComparison;
 
@@ -92,6 +94,7 @@ namespace UE::ImageWidgets
 				.Label(this, &SImageViewportToolbar::GetZoomMenuLabel)
 				.OnGetMenuContent(this, &SImageViewportToolbar::MakeZoomMenu)
 				.IsEnabled_Lambda([this] { return HasImage.Execute(); })
+				.Visibility(this, &SImageViewportToolbar::GetZoomMenuVisibility)
 			);
 
 			ToolbarBuilder.AddSeparator();
@@ -126,26 +129,27 @@ namespace UE::ImageWidgets
 				auto GetTextButton = [this](const FString& Label, const FCheckBoxStyle* ButtonStyle, FImageABComparison::EAorB AorB)
 				{
 					return SNew(SCheckBox)
-					.Style(ButtonStyle)
-					.IsEnabled_Lambda([this, AorB]
-						{
-							return ABComparison->CanSetABComparison(AorB);
-						})
-					.IsChecked(ABComparison->ABComparisonIsSet(AorB))
-					.OnCheckStateChanged_Lambda([this, AorB](const ECheckBoxState State)
-					    {
-							ABComparison->SetABComparison(AorB, State != ECheckBoxState::Checked ? FGuid() : ImageGuid.Execute());
-					    })
-					[
-						SNew(STextBlock)
-							.Font(FAppStyle::GetFontStyle("EditorViewportToolBar.Font"))
-							.Text(FText::FromString(Label))
-							.ToolTipText_Lambda([&ABComparison = ABComparison, AorB]
-								{
-									return ABComparison->ABComparisonIsSet(AorB) ? ABComparison->GetName(AorB) : FText();
-								})
-							.Margin(FMargin(2.0f, 0.0f))
-					];
+						.Style(ButtonStyle)
+						.Visibility(this, &SImageViewportToolbar::GetABVisibility)
+						.IsEnabled_Lambda([this, AorB]
+							{
+								return ABComparison->CanSetABComparison(AorB);
+							})
+						.IsChecked(ABComparison->ABComparisonIsSet(AorB))
+						.OnCheckStateChanged_Lambda([this, AorB](const ECheckBoxState State)
+						    {
+								ABComparison->SetABComparison(AorB, State != ECheckBoxState::Checked ? FGuid() : ImageGuid.Execute());
+						    })
+						[
+							SNew(STextBlock)
+								.Font(FAppStyle::GetFontStyle("EditorViewportToolBar.Font"))
+								.Text(FText::FromString(Label))
+								.ToolTipText_Lambda([&ABComparison = ABComparison, AorB]
+									{
+										return ABComparison->ABComparisonIsSet(AorB) ? ABComparison->GetName(AorB) : FText();
+									})
+								.Margin(FMargin(2.0f, 0.0f))
+						];
 				};
 
 				const FCheckBoxStyle* ButtonStyleStart = &FAppStyle::Get().GetWidgetStyle<FCheckBoxStyle>("EditorViewportToolBar.ToggleButton.Start");
@@ -230,9 +234,19 @@ namespace UE::ImageWidgets
 		return MenuBuilder.MakeWidget();
 	}
 
+	EVisibility SImageViewportToolbar::GetZoomMenuVisibility() const
+	{
+		return !GetOverlaySettings.Execute().bDisableZoomButton ? EVisibility::Visible : EVisibility::Collapsed;
+	}
+
 	EVisibility SImageViewportToolbar::GetMipMenuVisibility() const
 	{
-		return NumMips.Execute() > 1 ? EVisibility::Visible : EVisibility::Collapsed;
+		return !GetOverlaySettings.Execute().bDisableMipButton && NumMips.Execute() > 1 ? EVisibility::Visible : EVisibility::Collapsed;
+	}
+
+	EVisibility SImageViewportToolbar::GetABVisibility() const
+	{
+		return !GetOverlaySettings.Execute().bDisableABComparisonButtons ? EVisibility::Visible : EVisibility::Collapsed;
 	}
 
 	FText SImageViewportToolbar::GetMipMenuLabel() const
diff --git a/Engine/Plugins/Experimental/ImageWidgets/Source/ImageWidgets/Private/SImageViewportToolbar.h b/Engine/Plugins/Experimental/ImageWidgets/Source/ImageWidgets/Private/SImageViewportToolbar.h
--- a/Engine/Plugins/Experimental/ImageWidgets/Source/ImageWidgets/Private/SImageViewportToolbar.h
+++ b/Engine/Plugins/Experimental/ImageWidgets/Source/ImageWidgets/Private/SImageViewportToolbar.h
@@ -27,6 +27,7 @@ namespace UE::ImageWidgets
 		DECLARE_DELEGATE_RetVal(bool, FHasImage)
 		DECLARE_DELEGATE_RetVal(int32, FNumMips)
 		DECLARE_DELEGATE_RetVal(FGuid, FImageGuid)
+		DECLARE_DELEGATE_RetVal(SImageViewport::FOverlaySettings, FGetOverlaySettings)
 
 		SLATE_BEGIN_ARGS(SImageViewportToolbar)	{}
 		SLATE_END_ARGS()
@@ -37,6 +38,7 @@ namespace UE::ImageWidgets
 			FHasImage HasImage;
 			FNumMips NumMips;
 			FImageGuid ImageGuid;
+			FGetOverlaySettings GetOverlaySettings;
 			FImageABComparison* ABComparison = nullptr;
 			TSharedPtr<FExtender> ToolbarExtender;
 		};
@@ -49,7 +51,9 @@ namespace UE::ImageWidgets
 		TSharedRef<SWidget> MakeCenterToolbar(const TSharedPtr<FExtender>& Extender, bool bEnableABComparison);
 		TSharedRef<SWidget> MakeRightToolbar(const TSharedPtr<FExtender>& Extender);
 
+		EVisibility GetZoomMenuVisibility() const;
 		EVisibility GetMipMenuVisibility() const;
+		EVisibility GetABVisibility() const;
 		FText GetMipMenuLabel() const;
 		TSharedRef<SWidget> MakeMipMenu() const;
 
@@ -62,6 +66,7 @@ namespace UE::ImageWidgets
 		FHasImage HasImage;
 		FNumMips NumMips;
 		FImageGuid ImageGuid;
+		FGetOverlaySettings GetOverlaySettings;
 		FImageABComparison* ABComparison = nullptr;
 	};
 }
diff --git a/Engine/Plugins/Experimental/ImageWidgets/Source/ImageWidgets/Public/SImageViewport.h b/Engine/Plugins/Experimental/ImageWidgets/Source/ImageWidgets/Public/SImageViewport.h
--- a/Engine/Plugins/Experimental/ImageWidgets/Source/ImageWidgets/Public/SImageViewport.h
+++ b/Engine/Plugins/Experimental/ImageWidgets/Source/ImageWidgets/Public/SImageViewport.h
@@ -108,16 +108,46 @@ namespace UE::ImageWidgets
 				Fill // Make the image fit within the viewport, and if it is smaller than the viewport, zoom in to fill the viewport.
 			};
 
+			/**
+			 * Delegate for custom input key event handling in the viewport client.
+			 * The delegate is called after the image viewports input key events were handled, e.g. zoom or view reset via F key, and before a call
+			 * to @see FEditorViewportClient::InputKey().
+			 *
+			 * @param	EventArgs - The Input event args.
+			 * @return	True to consume the key event, false to pass it on.
+			 */
+			DECLARE_DELEGATE_RetVal_OneParam(bool, FOnInputKey, const FInputKeyEventArgs& EventArgs);
+
 			/** Zoom mode that gets set on viewport construction and whenever the viewport controller is reset via @see ResetController(). */
 			EDefaultZoomMode DefaultZoomMode = EDefaultZoomMode::Fit;
+
+			/** Delegate for custom input key handling. */
+			FOnInputKey OnInputKey;
 		};
 
-		DECLARE_DELEGATE(FOnLeftMouseButtonPressed);
-		DECLARE_DELEGATE(FOnLeftMouseButtonReleased);
+		/**
+		 * Settings related to the viewport overlay.
+		 */
+		struct FOverlaySettings
+		{
+			/** Do not show the zoom button in the left toolbar. */
+			bool bDisableZoomButton = false;
+
+			/** Do not show the MIP button in the left toolbar. */
+			bool bDisableMipButton = false;
+
+			/** Do not show the AB comparison buttons in the center toolbar. */
+			bool bDisableABComparisonButtons = false;
+
+			/** Do not show the bottom left status bar. */
+			bool bDisableStatusBarLeft = false;
+
+			/** Do not show the bottom left status bar. */
+			bool bDisableStatusBarRight = false;
+		};
 
 		SLATE_BEGIN_ARGS(SImageViewport)
 			: _bABComparisonEnabled(false)
-			, _MouseCaptureMode(EMouseCaptureMode::CapturePermanently)
 			{
 			}
 
@@ -130,21 +160,15 @@ namespace UE::ImageWidgets
 			/** Settings for drawing viewport contents other than the actual image */
 			SLATE_ATTRIBUTE(FDrawSettings, DrawSettings)
 
+			/** Settings for customizing the viewport overlay */
+			SLATE_ATTRIBUTE(FOverlaySettings, OverlaySettings)
+
 			/** Enables AB comparison controls in the toolbar */
 			SLATE_ARGUMENT(bool, bABComparisonEnabled)
 
 			/** Settings for controlling the viewport */
 			SLATE_ARGUMENT(FControllerSettings, ControllerSettings)
 
-			/** Left mouse button pressed event bubbled up from the viewport */
-			SLATE_EVENT(FOnLeftMouseButtonPressed, OnLeftMouseButtonPressed)
-
-			/** Left mouse button released event bubbled up from the viewport */
-			SLATE_EVENT(FOnLeftMouseButtonReleased, OnLeftMouseButtonReleased)
-
-			/** Settings for handling the mouse capture */
-			SLATE_ARGUMENT(EMouseCaptureMode, MouseCaptureMode)
-
 		SLATE_END_ARGS()
 
 		IMAGEWIDGETS_API SImageViewport();
@@ -169,10 +193,10 @@ namespace UE::ImageWidgets
 		 */
 		struct FPixelCoordinatesUnderCursorResult
 		{
-			/** Indicates that the cursor position is currently valid. This is set to false, for example, when the cursor is outside of the widget. */
+			/** Indicates that the cursor position is currently valid. This is set to false, for example, when the cursor is outside the widget. */
 			bool bIsValid = false;
 
-			/** Pixel coordinates under the cursor relative to the image rectangle size and placement. Note that the coordinates might be outside of the image
+			/** Pixel coordinates under the cursor relative to the image rectangle size and placement. Note that the coordinates might be outside the image
 			 * rectangle, i.e. values might be negative or larger than the image size. */
 			FVector2d Coordinates;
 		};
@@ -220,9 +244,12 @@ namespace UE::ImageWidgets
 		/** Provides the text for the resolution display in the status bar. */
 		FText GetResolutionLabel() const;
 
-		/** Makes the draw setting available either as fixed values or via a callback to the outside of the viewport. */
+		/** Makes the draw settings available either as fixed values or via a callback to the outside of the viewport. */
 		TAttribute<FDrawSettings> DrawSettings;
 
+		/** Makes the overlay settings available either as fixed values or via a callback to the outside of the viewport. */
+		TAttribute<FOverlaySettings> OverlaySettings;
+
 		/** Flag that determines is AB comparison widgets are enabled or not. The value does not change after the call to @see Construct(). */
 		bool bABComparisonEnabled = false;
 
@@ -243,9 +270,5 @@ namespace UE::ImageWidgets
 
 		/** Status bar extensions provided by the call to @see Construct(). This pointer is reset after the extensions were applied during construction. */
 		TSharedPtr<FStatusBarExtender> StatusBarExtender;
-
-		/** Left Mouse buttons pressed and released events triggered to tell the widget containing the viewport. These are initialized from the constructor.*/
-		FOnLeftMouseButtonPressed OnLeftMouseButtonPressed;
-		FOnLeftMouseButtonReleased OnLeftMouseButtonReleased;
 	};
 }
