ImageWidgets:
* Add function to explicitly request redrawing the image viewport
* Add callback to populate image catalog context menu
* Add Slate argument for setting image catalog selection mode
* Add image catalog functions to add an item before an existing one, delete items, pin/unpin items, and to query the catalog contents in various ways
* Update ColorViewer sample code to utilize the new functionality

diff --git a/Engine/Plugins/Experimental/ImageWidgets/Source/ImageWidgets/Private/ColorViewerSample/ColorViewer.cpp b/Engine/Plugins/Experimental/ImageWidgets/Source/ImageWidgets/Private/ColorViewerSample/ColorViewer.cpp
--- a/Engine/Plugins/Experimental/ImageWidgets/Source/ImageWidgets/Private/ColorViewerSample/ColorViewer.cpp
+++ b/Engine/Plugins/Experimental/ImageWidgets/Source/ImageWidgets/Private/ColorViewerSample/ColorViewer.cpp
@@ -11,13 +11,33 @@
 
 namespace UE::ImageWidgets::Sample
 {
+namespace Private
+{
+	static constexpr FGuid InvalidGuid{0, 0, 0, 0}; 
+
+	bool IsValid(const FGuid& Guid)
+	{
+		return Guid.A != 0;
+	}
+
+	void Invalidate(FGuid& Guid)
+	{
+		Guid.A = 0;
+	}
+	
+	int32 GetIndex(const FGuid& Guid)
+	{
+		return Guid.B;
+	}
+}
+
 	IImageViewer::FImageInfo FColorViewer::GetCurrentImageInfo() const
 	{
-		if (ColorIsValid(SelectedColorIndex))
+		if (ColorIndexIsValid(SelectedColorIndex) && ColorGuidIsValid(Colors[SelectedColorIndex].Guid))
 		{
 			return {Colors[SelectedColorIndex].Guid, ImageSize, 0, true};
 		}
-		return {FGuid(), FIntPoint::ZeroValue, 0, false};
+		return {Private::InvalidGuid, FIntPoint::ZeroValue, 0, false};
 	}
 
 	void FColorViewer::DrawCurrentImage(FViewport* Viewport, FCanvas* Canvas, const FDrawProperties& Properties)
@@ -35,7 +55,7 @@ namespace UE::ImageWidgets::Sample
 
 	TOptional<TVariant<FColor, FLinearColor>> FColorViewer::GetCurrentImagePixelColor(FIntPoint PixelCoords, int32 MipLevel) const
 	{
-		if (ColorIsValid(SelectedColorIndex))
+		if (ColorIndexIsValid(SelectedColorIndex) && ColorGuidIsValid(Colors[SelectedColorIndex].Guid))
 		{
 			// Returns the current color as float values.
 			// In a less trivial use case, the pixel coordinates and potentially the MIP level would be needed to look up the color value.
@@ -46,22 +66,22 @@ namespace UE::ImageWidgets::Sample
 
 	void FColorViewer::OnImageSelected(const FGuid& Guid)
 	{
-		if (ColorIsValid(Guid.B) && Colors[Guid.B].Guid == Guid)
+		if (ColorGuidIsValid(Guid))
 		{
-			SelectedColorIndex = Guid.B;
+			SelectedColorIndex = Private::GetIndex(Guid);
 		}
 	}
 
 	bool FColorViewer::IsValidImage(const FGuid& Guid) const
 	{
-		return ColorIsValid(Guid.B) && Colors[Guid.B].Guid == Guid;
+		return ColorGuidIsValid(Guid);
 	}
 
 	FText FColorViewer::GetImageName(const FGuid& Guid) const
 	{
-		if (IsValidImage(Guid))
+		if (ColorGuidIsValid(Guid))
 		{
-			const FColor& Color = Colors[Guid.B].Color;
+			const FColor& Color = Colors[Private::GetIndex(Guid)].Color;
 			const FString HexColor = FString::Printf(TEXT("#%02X%02X%02X"), Color.R, Color.G, Color.B);
 			return FText::FromString(HexColor);
 		}
@@ -79,9 +99,24 @@ namespace UE::ImageWidgets::Sample
 		return &Colors[SelectedColorIndex];
 	}
 
+	void FColorViewer::RemoveColor(const FGuid& Guid)
+	{
+		if (ColorGuidIsValid(Guid))
+		{
+			const int32 Index = Private::GetIndex(Guid);
+
+			Private::Invalidate(Colors[Index].Guid);
+
+			if (SelectedColorIndex == Index)
+			{
+				SelectedColorIndex = INDEX_NONE;
+			}
+		}
+	}
+
 	const FColorViewer::FColorItem* FColorViewer::RandomizeColor()
 	{
-		if (ColorIsValid(SelectedColorIndex))
+		if (ColorIndexIsValid(SelectedColorIndex))
 		{
 			auto Random = []
 			{
@@ -111,14 +146,19 @@ namespace UE::ImageWidgets::Sample
 		return ToneMapping.GetToneMappedColor(Color);
 	}
 
-	bool FColorViewer::ColorIsValid(int32 Index) const
+	bool FColorViewer::ColorIndexIsValid(int32 Index) const
 	{
 		return 0 <= Index && Index < Colors.Num();
 	}
 
+	bool FColorViewer::ColorGuidIsValid(const FGuid& Guid) const
+	{
+		return Guid.A != 0 && ColorIndexIsValid(Guid.B) && Colors[Guid.B].Guid == Guid;
+	}
+
 	void FColorViewer::DrawImage(int32 Index, FCanvas* Canvas, const FDrawProperties::FPlacement& Placement, const FVector2d& UV0, const FVector2d& UV1) const
 	{
-		if (ColorIsValid(Index))
+		if (ColorIndexIsValid(Index))
 		{
 			// Get color value after tone mapping.
 			const FLinearColor ToneMappedColor = ToneMapping.GetToneMappedColor(Colors[Index].Color);
diff --git a/Engine/Plugins/Experimental/ImageWidgets/Source/ImageWidgets/Private/ColorViewerSample/ColorViewer.h b/Engine/Plugins/Experimental/ImageWidgets/Source/ImageWidgets/Private/ColorViewerSample/ColorViewer.h
--- a/Engine/Plugins/Experimental/ImageWidgets/Source/ImageWidgets/Private/ColorViewerSample/ColorViewer.h
+++ b/Engine/Plugins/Experimental/ImageWidgets/Source/ImageWidgets/Private/ColorViewerSample/ColorViewer.h
@@ -67,6 +67,9 @@ namespace UE::ImageWidgets::Sample
 
 		/** Adds a color item. */
 		const FColorItem* AddColor();
+
+		/** Removes a color item. */
+		void RemoveColor(const FGuid& Guid);
 		
 		/** Sets a random color as the current "image" as a simple proxy for the image content changing and/or users choosing different images to display. */
 		const FColorItem* RandomizeColor();
@@ -84,7 +87,10 @@ namespace UE::ImageWidgets::Sample
 
 	private:
 		/** Checks if a given index is a valid image. */
-		bool ColorIsValid(int32 Index) const;
+		bool ColorIndexIsValid(int32 Index) const;
+
+		/** Checks if a given guid is a valid image. */
+		bool ColorGuidIsValid(const FGuid& Guid) const;
 
 		/** Draws the color image with the given index. The UVs determine if all or only a part of the image is drawn, i.e. for AB comparisons. */
 		void DrawImage(int32 Index, FCanvas* Canvas, const FDrawProperties::FPlacement& Placement, const FVector2d& UV0, const FVector2d& UV1) const;
diff --git a/Engine/Plugins/Experimental/ImageWidgets/Source/ImageWidgets/Private/ColorViewerSample/ColorViewerWidget.cpp b/Engine/Plugins/Experimental/ImageWidgets/Source/ImageWidgets/Private/ColorViewerSample/ColorViewerWidget.cpp
--- a/Engine/Plugins/Experimental/ImageWidgets/Source/ImageWidgets/Private/ColorViewerSample/ColorViewerWidget.cpp
+++ b/Engine/Plugins/Experimental/ImageWidgets/Source/ImageWidgets/Private/ColorViewerSample/ColorViewerWidget.cpp
@@ -6,6 +6,9 @@
 
 #include "ColorViewerCommands.h"
 #include "ColorViewerStyle.h"
+#include "Algo/AllOf.h"
+#include "Algo/NoneOf.h"
+#include "Algo/Sort.h"
 #include "Framework/MultiBox/MultiBoxBuilder.h"
 
 #define LOCTEXT_NAMESPACE "ColorViewerWidget"
@@ -36,7 +39,20 @@ namespace UE::ImageWidgets::Sample
 						.Value(0.0f)
 						[
 							SAssignNew(Catalog, ImageWidgets::SImageCatalog)
-								.OnItemSelected_Lambda([this](const FGuid& ImageGuid) { ColorViewer->OnImageSelected(ImageGuid); })
+								.PinnedItemsHeading(LOCTEXT("PinnedColors", "Pinned Colors"))
+								.ItemsHeading(LOCTEXT("Colors", "Colors"))
+								.OnItemSelected_Lambda([&ColorViewer = ColorViewer, &Viewport = Viewport](const FGuid& ImageGuid)
+								{
+									if (ColorViewer)
+									{
+										ColorViewer->OnImageSelected(ImageGuid);
+									}
+									if (Viewport)
+									{
+										Viewport->RequestRedraw();
+									}
+								})
+								.OnGetContextMenu(this, &SColorViewerWidget::GetContextMenu)
 						]
 					+ SSplitter::Slot()
 						.Value(1.0f)
@@ -138,6 +154,85 @@ namespace UE::ImageWidgets::Sample
 		}
 	}
 
+	TSharedPtr<SWidget> SColorViewerWidget::GetContextMenu(const TArray<FGuid>& Guids) const
+	{
+		checkSlow(!Guids.IsEmpty());
+		checkSlow(Algo::AllOf(Guids, [&Catalog = Catalog](const FGuid& Guid) { return Catalog->ItemIsPinned(Guid); })
+			|| Algo::NoneOf(Guids, [&Catalog = Catalog] (const FGuid& Guid) { return Catalog->ItemIsPinned(Guid); }));
+
+		FMenuBuilder MenuBuilder(true, nullptr);
+
+		if (Catalog->ItemIsPinned(Guids[0]))
+		{
+			MenuBuilder.AddMenuEntry(
+				LOCTEXT("UnpinItems", "Unpin"),
+				LOCTEXT("UnpinItemsTooltip", "Removes the selected item(s) from the list of pinned items."),
+				FSlateIcon(FAppStyle::GetAppStyleSetName(), "Icons.Unpinned"),
+				FUIAction(FExecuteAction::CreateLambda(
+					[Guids, Catalog = Catalog]
+					{
+						if (Catalog)
+						{
+							for (const FGuid& Guid : Guids)
+							{
+								Catalog->UnpinItem(Guid);
+							}
+							Catalog->SelectItem(Guids[0]);
+						}
+					})));
+		}
+		else
+		{
+			MenuBuilder.AddMenuEntry(
+				LOCTEXT("PinItems", "Pin"),
+				LOCTEXT("PinItemsTooltip", "Adds the selected item(s) to the list of pinned items."),
+				FSlateIcon(FAppStyle::GetAppStyleSetName(), "Icons.Pinned"),
+				FUIAction(
+					FExecuteAction::CreateLambda(
+						[Guids, Catalog = Catalog]
+						{
+							if (Catalog)
+							{
+								for (const FGuid& Guid : Guids)
+								{
+									Catalog->PinItem(Guid);
+								}
+								Catalog->SelectItem(Guids[0]);
+							}
+						})));
+		}
+
+#if 0
+		MenuBuilder.AddMenuEntry(
+			LOCTEXT("RenameItem", "Rename"),
+			LOCTEXT("RenameItemTooltip", "Renames a single selected item."),
+			FSlateIcon(FAppStyle::GetAppStyleSetName(), "GenericCommands.Rename"),
+			FUIAction(
+				FExecuteAction::CreateLambda([]{}),
+				FCanExecuteAction::CreateLambda([bSingleItem = Guids.Num() == 1] { return bSingleItem; })));
+#endif
+
+		MenuBuilder.AddSeparator();
+
+		MenuBuilder.AddMenuEntry(
+			LOCTEXT("DeleteItems", "Delete"), LOCTEXT("DeleteItemsTooltip", "Deletes the selected item(s) from the catalog."),
+			FSlateIcon(FAppStyle::GetAppStyleSetName(), "GenericCommands.Delete"),
+			FUIAction(FExecuteAction::CreateLambda(
+				[Guids, Catalog = Catalog, ColorViewer = ColorViewer]
+				{
+					if (Catalog && ColorViewer)
+					{
+						for (const FGuid& Guid : Guids)
+						{
+							Catalog->RemoveItem(Guid);
+							ColorViewer->RemoveColor(Guid);
+						}
+					}
+				})));
+
+		return MenuBuilder.MakeWidget();
+	}
+
 	void SColorViewerWidget::BindCommands()
 	{
 		const FColorViewerCommands& Commands = FColorViewerCommands::Get();
@@ -158,7 +253,7 @@ namespace UE::ImageWidgets::Sample
 			Commands.ToneMappingRGB,
 			FExecuteAction::CreateSP(ColorViewer.ToSharedRef(), &FColorViewer::SetToneMapping, FToneMapping::EMode::RGB),
 			FCanExecuteAction(),
-			FIsActionChecked::CreateLambda([this]
+			FIsActionChecked::CreateLambda([ColorViewer = ColorViewer]
 			{
 				return ColorViewer->GetToneMapping() == FToneMapping::EMode::RGB;
 			})
@@ -168,7 +263,7 @@ namespace UE::ImageWidgets::Sample
 			Commands.ToneMappingLum,
 			FExecuteAction::CreateSP(ColorViewer.ToSharedRef(), &FColorViewer::SetToneMapping, FToneMapping::EMode::Lum),
 			FCanExecuteAction(),
-			FIsActionChecked::CreateLambda([this]
+			FIsActionChecked::CreateLambda([ColorViewer = ColorViewer]
 			{
 				return ColorViewer->GetToneMapping() == FToneMapping::EMode::Lum;
 			})
diff --git a/Engine/Plugins/Experimental/ImageWidgets/Source/ImageWidgets/Private/ColorViewerSample/ColorViewerWidget.h b/Engine/Plugins/Experimental/ImageWidgets/Source/ImageWidgets/Private/ColorViewerSample/ColorViewerWidget.h
--- a/Engine/Plugins/Experimental/ImageWidgets/Source/ImageWidgets/Private/ColorViewerSample/ColorViewerWidget.h
+++ b/Engine/Plugins/Experimental/ImageWidgets/Source/ImageWidgets/Private/ColorViewerSample/ColorViewerWidget.h
@@ -43,6 +43,9 @@ namespace UE::ImageWidgets::Sample
 		/** Choose a random color for the current entry. */
 		void RandomizeColor();
 
+		/** Creates a context menu for a set of selected guids. */
+		TSharedPtr<SWidget> GetContextMenu(const TArray<FGuid>& Guids) const;
+
 		/** The image viewer implementation that contains the image data and renders the image. */
 		TSharedPtr<FColorViewer> ColorViewer;
 
diff --git a/Engine/Plugins/Experimental/ImageWidgets/Source/ImageWidgets/Private/SImageCatalog.cpp b/Engine/Plugins/Experimental/ImageWidgets/Source/ImageWidgets/Private/SImageCatalog.cpp
--- a/Engine/Plugins/Experimental/ImageWidgets/Source/ImageWidgets/Private/SImageCatalog.cpp
+++ b/Engine/Plugins/Experimental/ImageWidgets/Source/ImageWidgets/Private/SImageCatalog.cpp
@@ -12,18 +12,262 @@
 
 namespace UE::ImageWidgets
 {
-FImageCatalogItemData::FImageCatalogItemData(FGuid Guid, const FSlateBrush& Brush, const FText& Name, const FText& Info, const FText& ToolTip)
+FImageCatalogItemData::FImageCatalogItemData(const FGuid Guid, const FSlateBrush& Brush, const FText& Name, const FText& Info, const FText& ToolTip)
 	: Guid(Guid), Thumbnail(Brush), Name(Name), Info(Info), ToolTip(ToolTip)
 {
 }
 
+using FItemType = TSharedPtr<FImageCatalogItemData>;
+	
+class FItemModel
+{
+public:
+
+	bool Add(const FItemType& Item, bool bIsPinned, const FGuid* BeforeThisGuid);
+	TTuple<bool, bool> Remove(const FGuid& Guid);
+
+	TTuple<const FItemType*, bool> GetItem(const FGuid& Guid) const;
+	const FItemType* GetItemAt(int32 Index, bool bIsPinned) const;
+	TOptional<TTuple<bool, int32>> GetItemIndex(const FGuid& Guid) const;
+	TOptional<FGuid> GetGuidAt(int32 Index, bool bIsPinned) const;
+	
+	bool IsPinned(const FGuid& Guid);
+	bool Pin(const FGuid& Guid);
+	bool Unpin(const FGuid& Guid);
+	
+	const TArray<FItemType>& GetPinnedItems() const { return PinnedItems; }
+	const TArray<FItemType>& GetUnpinnedItems() const { return UnpinnedItems; }
+	bool HasPinnedItems() const { return PinnedItems.Num() > 0; }
+	bool HasUnpinnedItems() const { return UnpinnedItems.Num() > 0; }
+	int32 NumPinnedItems() const { return PinnedItems.Num(); }
+	int32 NumUnpinnedItems() const { return UnpinnedItems.Num(); }
+
+	void SortSelection(TArray<FGuid>& Selection);
+
+private:
+
+	struct FLookupData
+	{
+		bool bIsPinned;	// Is this item pinned?
+		int32 Index;	// Index in the respective array.
+	};
+
+	const FLookupData* FindLookupData(const FGuid& Guid) const;
+	FLookupData* FindLookupData(const FGuid& Guid)
+	{
+		return const_cast<FLookupData*>(const_cast<const FItemModel*>(this)->FindLookupData(Guid));
+	}
+
+	const FItemType* GetItem(const FLookupData& LookupData) const;
+
+	void SwapBetweenPinnedAndUnpinned(FLookupData& LookupData);
+	void UpdateMappingIndices(const TArray<FItemType>& Container, int32 FirstIndex);
+
+	TArray<FItemType> PinnedItems;
+	TArray<FItemType> UnpinnedItems;
+	TMap<FGuid, FLookupData> GuidToLookupDataMapping;
+};
+
+bool FItemModel::Add(const FItemType& Item, const bool bIsPinned, const FGuid* BeforeThisGuid)
+{
+	if (!GuidToLookupDataMapping.Contains(Item->Guid))
+	{
+		TArray<FItemType>& Container = bIsPinned ? PinnedItems : UnpinnedItems;
+
+		const int32 Index = [this, &Item, bIsPinned, BeforeThisGuid, &Container]
+		{
+			if (BeforeThisGuid)
+			{
+				if (FLookupData *const LookupData = FindLookupData(*BeforeThisGuid))
+				{
+					if (LookupData->bIsPinned == bIsPinned)
+					{
+						// Set index for added item.
+						const int32 NewIndex = LookupData->Index;
+
+						// Increase index for the item we push back.
+						++LookupData->Index;
+
+						// Add item at the new index.
+						Container.EmplaceAt(NewIndex, Item);
+
+						// Update all lookup data for items that come after the item we used to determine the insert location.
+						// This way we save the effort for finding the same lookup data again.
+						UpdateMappingIndices(Container, LookupData->Index + 1);
+
+						// Tell the outside where the new item was added.
+						return NewIndex;
+					}
+				}
+			}
+			
+			return Container.Add(Item);
+		}();
+
+		GuidToLookupDataMapping.Add(Item->Guid, {bIsPinned, Index});
+		return true;
+	}
+
+	return false;
+}
+
+TTuple<bool, bool> FItemModel::Remove(const FGuid& Guid)
+{
+	if (const FLookupData *const LookupData = FindLookupData(Guid))
+	{
+		const bool bIsPinned = LookupData->bIsPinned;
+		TArray<FItemType>& Container = bIsPinned ? PinnedItems : UnpinnedItems;
+		Container.RemoveAt(LookupData->Index, EAllowShrinking::No);
+		GuidToLookupDataMapping.Remove(Guid);
+		UpdateMappingIndices(Container, LookupData->Index);
+		return {true, bIsPinned};
+	}
+
+	return {false, false};
+}
+
+TTuple<const FItemType*, bool> FItemModel::GetItem(const FGuid& Guid) const
+{
+	if (const FLookupData *const LookupData = FindLookupData(Guid))
+	{
+		return {GetItem(*LookupData), LookupData->bIsPinned};
+	}
+
+	return {nullptr, false};
+}
+
+const FItemType* FItemModel::GetItemAt(const int32 Index, const bool bIsPinned) const
+{
+	const TArray<FItemType>& Container = bIsPinned ? PinnedItems : UnpinnedItems;
+	if (0 <= Index && Index < Container.Num())
+	{
+		return &Container[Index];
+	}
+
+	return nullptr;
+}
+
+TOptional<TTuple<bool, int32>> FItemModel::GetItemIndex(const FGuid& Guid) const
+{
+	if (const FLookupData *const LookupData = FindLookupData(Guid))
+	{
+		return {{LookupData->bIsPinned, LookupData->Index}};
+	}
+
+	return {};
+}
+
+TOptional<FGuid> FItemModel::GetGuidAt(const int32 Index, const bool bIsPinned) const
+{
+	const TArray<FItemType>& Container = bIsPinned ? PinnedItems : UnpinnedItems;
+	if (0 <= Index && Index < Container.Num())
+	{
+		return Container[Index]->Guid;
+	}
+
+	return {};
+}
+
+bool FItemModel::IsPinned(const FGuid& Guid)
+{
+	const FLookupData *const LookupData = FindLookupData(Guid);
+	return LookupData && LookupData->bIsPinned;
+}
+
+bool FItemModel::Pin(const FGuid& Guid)
+{
+	FLookupData *const LookupData = FindLookupData(Guid);
+	if (LookupData && !LookupData->bIsPinned)
+	{
+		SwapBetweenPinnedAndUnpinned(*LookupData);
+		return true;
+	}
+
+	return false;
+}
+
+bool FItemModel::Unpin(const FGuid& Guid)
+{
+	FLookupData *const LookupData = FindLookupData(Guid);
+	if (LookupData && LookupData->bIsPinned)
+	{
+		SwapBetweenPinnedAndUnpinned(*LookupData);
+		return true;
+	}
+
+	return false;
+}
+
+void FItemModel::SortSelection(TArray<FGuid>& Selection)
+{
+	Algo::Sort(Selection, [this, NumPinned = PinnedItems.Num()](const FGuid& A, const FGuid& B)
+	{
+		const FLookupData *const LookupDataA = FindLookupData(A);
+		const FLookupData *const LookupDataB = FindLookupData(B);
+
+		if (LookupDataA && !LookupDataB) return true;
+		if (!LookupDataA && LookupDataB) return false;
+		if (!LookupDataA && !LookupDataB) return true;
+
+		const int32 PositionA = (LookupDataA->bIsPinned ? 0 : NumPinned) + LookupDataA->Index;
+		const int32 PositionB = (LookupDataB->bIsPinned ? 0 : NumPinned) + LookupDataB->Index;
+		
+		return PositionA < PositionB;
+	});
+}
+
+const FItemModel::FLookupData* FItemModel::FindLookupData(const FGuid& Guid) const
+{
+	if (const FLookupData* LookupData = GuidToLookupDataMapping.Find(Guid))
+	{
+		return LookupData;
+	}
+
+	UE_LOG(LogImageWidgets, Warning, TEXT("Cannot find catalog item for guid '%s'."), *Guid.ToString());
+	return nullptr;
+}
+
+const FItemType* FItemModel::GetItem(const FLookupData& LookupData) const
+{
+	const TArray<FItemType> *const Container = LookupData.bIsPinned ? &PinnedItems : &UnpinnedItems;
+	check(0 <= LookupData.Index && LookupData.Index < Container->Num());
+	return &(*Container)[LookupData.Index];
+}
+
+void FItemModel::SwapBetweenPinnedAndUnpinned(FLookupData& LookupData)
+{
+	TArray<FItemType>& From = LookupData.bIsPinned ? PinnedItems : UnpinnedItems;
+	TArray<FItemType>& To = LookupData.bIsPinned ? UnpinnedItems : PinnedItems;
+	
+	const int32 NewIndex = To.Emplace(static_cast<FItemType&&>(From[LookupData.Index]));
+	From.RemoveAt(LookupData.Index, EAllowShrinking::No);
+
+	UpdateMappingIndices(From, LookupData.Index);
+
+	LookupData.Index = NewIndex;
+	LookupData.bIsPinned = !LookupData.bIsPinned;
+}
+
+void FItemModel::UpdateMappingIndices(const TArray<FItemType>& Container, const int32 FirstIndex)
+{
+	for (int32 Index = FirstIndex, Num = Container.Num(); Index < Num; ++Index)
+	{
+		FLookupData *const LookupData = GuidToLookupDataMapping.Find(Container[Index]->Guid);
+		check(LookupData);
+		LookupData->Index = Index;
+	}
+}
+
 void SImageCatalog::Construct(const FArguments& InArgs)
 {
+	Model = MakePimpl<FItemModel>();
+	
+	ItemsHeading = InArgs._ItemsHeading;
+	PinnedItemsHeading = InArgs._PinnedItemsHeading;
 	OnItemSelected = InArgs._OnItemSelected;
+	OnGetContextMenu = InArgs._OnGetContextMenu;
 
-	Items.Reserve(1000);
-
-	const auto GenerateItemRow = [](const TSharedPtr<FImageCatalogItemData>& ItemData, const TSharedRef<STableViewBase>& OwnerTable)
+	const auto GenerateItemRow = [](const FItemType& ItemData, const TSharedRef<STableViewBase>& OwnerTable)
 	{
 		static const FTableRowStyle TableRowStyle = []
 		{
@@ -38,7 +282,7 @@ void SImageCatalog::Construct(const FArguments& InArgs)
 		TSharedPtr<SImageCatalogItem> ItemWidget;
 		SAssignNew(ItemWidget, SImageCatalogItem, ItemData);
 
-		return SNew(STableRow<TSharedPtr<FImageCatalogItemData>>, OwnerTable)
+		return SNew(STableRow<FItemType>, OwnerTable)
 			.Style(&TableRowStyle)
 			.ShowSelection(true)
 			[
@@ -46,42 +290,51 @@ void SImageCatalog::Construct(const FArguments& InArgs)
 			];
 	};
 
-	const auto HasPinnedItems = [this]
+	const auto GetPinnedItemsVisibility = [&Model = Model]
 	{
-		return PinnedItems.IsEmpty() ? EVisibility::Collapsed : EVisibility::Visible;
+		return Model->HasPinnedItems() ? EVisibility::Visible : EVisibility::Collapsed;
 	};
 
-	const auto PinnedItemSelectionChanged = [this](const TSharedPtr<FImageCatalogItemData>& Item, ESelectInfo::Type SelectInfo)
+	const auto GetItemsVisibility = [&Model = Model]
+	{
+		return Model->HasUnpinnedItems() ? EVisibility::Visible : EVisibility::Collapsed;
+	};
+
+	const auto PinnedItemsSelectionChanged = [this](const FItemType& Item, ESelectInfo::Type SelectInfo)
 	{
 		// Note that Item might be a nullptr since this callback is also executed when clearing the selection of the list.
-		if (Item.IsValid() && OnItemSelected.IsBound())
+		if (Item.IsValid())
 		{
 			ItemsListView->ClearSelection();
-			OnItemSelected.Execute(Item->Guid);
+			OnItemSelected.ExecuteIfBound(Item->Guid);
 		}
 	};
 
-	const auto ItemSelectionChanged = [this](const TSharedPtr<FImageCatalogItemData>& Item, ESelectInfo::Type SelectInfo)
+	const auto ItemsSelectionChanged = [this](const FItemType& Item, ESelectInfo::Type SelectInfo)
 	{
 		// Note that Item might be a nullptr since this callback is also executed when clearing the selection of the list.
-		if (Item.IsValid() && OnItemSelected.IsBound())
+		if (Item.IsValid())
 		{
 			PinnedItemsListView->ClearSelection();
-			OnItemSelected.Execute(Item->Guid);
+			OnItemSelected.ExecuteIfBound(Item->Guid);
 		}
 	};
 
-	SAssignNew(PinnedItemsListView, SListView<TSharedPtr<FImageCatalogItemData>>)
-					.ListItemsSource(&PinnedItems)
+	SAssignNew(PinnedItemsListView, SListView<FItemType>)
+					.ListItemsSource(&Model->GetPinnedItems())
+					.OnContextMenuOpening(this, &SImageCatalog::OnContextMenuOpening)
 					.OnGenerateRow_Lambda(GenerateItemRow)
-					.OnSelectionChanged_Lambda(PinnedItemSelectionChanged)
+					.SelectionMode(InArgs._SelectionMode)
+					.OnSelectionChanged_Lambda(PinnedItemsSelectionChanged)
 					.ClearSelectionOnClick(false)
-					.Visibility_Lambda(HasPinnedItems);
+					.Visibility_Lambda(GetPinnedItemsVisibility);
 
-	SAssignNew(ItemsListView, SListView<TSharedPtr<FImageCatalogItemData>>)
-					.ListItemsSource(&Items)
+	SAssignNew(ItemsListView, SListView<FItemType>)
+					.ListItemsSource(&Model->GetUnpinnedItems())
+					.OnContextMenuOpening(this, &SImageCatalog::OnContextMenuOpening)
 					.OnGenerateRow_Lambda(GenerateItemRow)
-					.OnSelectionChanged_Lambda(ItemSelectionChanged)
+					.SelectionMode(InArgs._SelectionMode)
+					.OnSelectionChanged_Lambda(ItemsSelectionChanged)
 					.ClearSelectionOnClick(false)
 					.ScrollbarVisibility(EVisibility::Visible);
 	
@@ -93,9 +346,12 @@ void SImageCatalog::Construct(const FArguments& InArgs)
 			.Padding(2.0f, 4.0f, 2.0f, 4.0f)
 			[
 				SNew(STextBlock)
-					.Text(InArgs._PinnedItemsHeading)
+					.Text_Lambda([&PinnedItemsHeading = PinnedItemsHeading] { return PinnedItemsHeading.Get({}); })
+					.Visibility_Lambda([&PinnedItemsHeading = PinnedItemsHeading, GetPinnedItemsVisibility]
+					{
+						return PinnedItemsHeading.Get({}).IsEmpty() ? EVisibility::Collapsed : GetPinnedItemsVisibility();
+					})
 					.TextStyle(FAppStyle::Get(), "DetailsView.CategoryTextStyle")
-					.Visibility(InArgs._PinnedItemsHeading.IsEmpty() ? EVisibility::Collapsed : EVisibility::Visible)
 			]
 		+ SVerticalBox::Slot()
 			.AutoHeight()
@@ -107,16 +363,19 @@ void SImageCatalog::Construct(const FArguments& InArgs)
 			[
 				SNew(SSeparator)
 					.Thickness(6.0f)
-					.Visibility_Lambda(HasPinnedItems)
+					.Visibility_Lambda(GetPinnedItemsVisibility)
 			]
 		+ SVerticalBox::Slot()
 			.AutoHeight()
 			.Padding(2.0f, 4.0f, 2.0f, 4.0f)
 			[
 				SNew(STextBlock)
-					.Text(InArgs._ItemsHeading)
+					.Text_Lambda([&ItemsHeading = ItemsHeading] { return ItemsHeading.Get({}); })
+					.Visibility_Lambda([&ItemsHeading = ItemsHeading, GetItemsVisibility]
+					{
+						return ItemsHeading.Get({}).IsEmpty() ? EVisibility::Collapsed : GetItemsVisibility();
+					})
 					.TextStyle(FAppStyle::Get(), "DetailsView.CategoryTextStyle")
-					.Visibility(InArgs._ItemsHeading.IsEmpty() ? EVisibility::Collapsed : EVisibility::Visible)
 			]
 		+ SVerticalBox::Slot()
 			[
@@ -125,28 +384,113 @@ void SImageCatalog::Construct(const FArguments& InArgs)
 	];
 }
 
-void SImageCatalog::AddItem(const TSharedPtr<FImageCatalogItemData>& Item)
+bool SImageCatalog::AddItem(const FItemType& Item)
+{
+	if (Model->Add(Item, false, nullptr))
+	{
+		ItemsListView->RequestListRefresh();
+		return true;
+	}
+
+	return false;
+}
+
+bool SImageCatalog::AddItem(const TSharedPtr<FImageCatalogItemData>& Item, const FGuid& BeforeItemWithThisGuid)
+{
+	if (Model->Add(Item, false, &BeforeItemWithThisGuid))
+	{
+		ItemsListView->RequestListRefresh();
+		return true;
+	}
+
+	return false;
+}
+
+bool SImageCatalog::AddPinnedItem(const FItemType& Item)
+{
+	if (Model->Add(Item, true, nullptr))
+	{
+		PinnedItemsListView->RequestListRefresh();
+		return true;
+	}
+
+	return false;
+}
+
+bool SImageCatalog::AddPinnedItem(const TSharedPtr<FImageCatalogItemData>& Item, const FGuid& BeforeItemWithThisGuid)
+{
+	if (Model->Add(Item, true, &BeforeItemWithThisGuid))
+	{
+		PinnedItemsListView->RequestListRefresh();
+		return true;
+	}
+
+	return false;
+}
+
+bool SImageCatalog::RemoveItem(const FGuid& Guid)
+{
+	const auto [bSuccess, bIsPinned] = Model->Remove(Guid);
+	if (bSuccess)
+	{
+		const TSharedPtr<SListView<FItemType>>& ListView = bIsPinned ? PinnedItemsListView : ItemsListView;
+		ListView->RequestListRefresh();
+		return true;
+	}
+
+	return false;
+}
+
+TSharedPtr<const FImageCatalogItemData> SImageCatalog::GetItem(const FGuid& Guid) const
+{
+	if (const FItemType* ItemPtr = Model->GetItem(Guid).Get<0>())
+	{
+		return *ItemPtr;
+	}
+	return {};
+}
+
+TOptional<TTuple<bool, int32>> SImageCatalog::GetItemIndex(const FGuid& Guid) const
+{
+	return Model->GetItemIndex(Guid);
+}
+
+TSharedPtr<const FImageCatalogItemData> SImageCatalog::GetItemAt(const int32 Index) const
+{
+	if (const FItemType* ItemPtr = Model->GetItemAt(Index, false))
+	{
+		return *ItemPtr;
+	}
+	return {};
+}
+
+TSharedPtr<const FImageCatalogItemData> SImageCatalog::GetPinnedItemAt(const int32 Index) const
+{
+	if (const FItemType* ItemPtr = Model->GetItemAt(Index, true))
+	{
+		return *ItemPtr;
+	}
+	return {};
+}
+
+TOptional<FGuid> SImageCatalog::GetItemGuidAt(const int32 Index) const
 {
-	Items.Add(Item);
-	GuidToItemMapping.Add(Item->Guid, TPair<const TSharedPtr<FImageCatalogItemData>, bool>(Item, false));
-	ItemsListView->ReGenerateItems(GetCachedGeometry());
+	return Model->GetGuidAt(Index, false);
 }
 
-void SImageCatalog::AddPinnedItem(const TSharedPtr<FImageCatalogItemData>& Item)
+TOptional<FGuid> SImageCatalog::GetPinnedItemGuidAt(const int32 Index) const 
 {
-	PinnedItems.Add(Item);
-	GuidToItemMapping.Add(Item->Guid, TPair<const TSharedPtr<FImageCatalogItemData>, bool>(Item, true));
-	PinnedItemsListView->ReGenerateItems(GetCachedGeometry());
+	return Model->GetGuidAt(Index, true);
 }
 
 int32 SImageCatalog::NumItems() const
 {
-	return Items.Num();
+	return Model->NumUnpinnedItems();
 }
 
 int32 SImageCatalog::NumPinnedItems() const
 {
-	return PinnedItems.Num();
+	return Model->NumPinnedItems();
 }
 
 int32 SImageCatalog::NumTotalItems() const
@@ -154,67 +498,109 @@ int32 SImageCatalog::NumTotalItems() const
 	return NumItems() + NumPinnedItems();
 }
 
+bool SImageCatalog::ItemIsPinned(const FGuid& Guid) const
+{
+	return Model->IsPinned(Guid);
+}
+	
+bool SImageCatalog::PinItem(const FGuid& Guid)
+{
+	const bool bSuccess = Model->Pin(Guid);
+	if (bSuccess)
+	{
+		PinnedItemsListView->RequestListRefresh();
+		ItemsListView->RequestListRefresh();
+	}
+	return bSuccess;
+}
+
+bool SImageCatalog::UnpinItem(const FGuid& Guid)
+{
+	const bool bSuccess = Model->Unpin(Guid);
+	if (bSuccess)
+	{
+		PinnedItemsListView->RequestListRefresh();
+		ItemsListView->RequestListRefresh();
+	}
+	return bSuccess;
+}
+
 void SImageCatalog::SelectItem(const FGuid& Guid)
 {
-	if (const TPair<const TSharedPtr<FImageCatalogItemData>, bool> *const ExistingItem = FindItemData(Guid))
+	const auto [ItemPtr, bIsPinned] = Model->GetItem(Guid);
+	if (ItemPtr)
 	{
 		ItemsListView->ClearSelection();
 		PinnedItemsListView->ClearSelection();
 
-		SListView<TSharedPtr<FImageCatalogItemData>>* ListView = ExistingItem->Get<1>() ? PinnedItemsListView.Get() : ItemsListView.Get();
-
-		ListView->SetItemSelection(ExistingItem->Get<0>(), true);
+		SListView<FItemType> *const ListView = bIsPinned ? PinnedItemsListView.Get() : ItemsListView.Get();
+		ListView->SetItemSelection(*ItemPtr, true);
 	}
 }
 
-const TPair<const TSharedPtr<FImageCatalogItemData>, bool>* SImageCatalog::FindItemData(const FGuid& Guid)
+TSharedPtr<SWidget> SImageCatalog::OnContextMenuOpening() const
 {
-	if (const TPair<const TSharedPtr<FImageCatalogItemData>, bool>* ItemPtr = GuidToItemMapping.Find(Guid))
+	if (!OnGetContextMenu.IsBound())
 	{
-		return ItemPtr;
+		return SNullWidget::NullWidget;
 	}
 
-	UE_LOG(LogImageWidgets, Warning, TEXT("Cannot find catalog item for guid '%s'."), *Guid.ToString());
-	return nullptr;
+	TArray<FGuid> SelectedGuids = [&PinnedItemsListView = PinnedItemsListView, &ItemsListView = ItemsListView]
+	{
+		TArray<FGuid> Guids;
+		auto GetItemGuid = [](const FItemType& Item) { return Item->Guid; };
+		Algo::Transform(PinnedItemsListView->GetSelectedItems(), Guids, GetItemGuid);
+		Algo::Transform(ItemsListView->GetSelectedItems(), Guids, GetItemGuid);
+		return Guids;
+	}();
+
+	if (SelectedGuids.IsEmpty())
+	{
+		return SNullWidget::NullWidget;
+	}
+
+	Model->SortSelection(SelectedGuids);
+
+	return OnGetContextMenu.Execute(SelectedGuids);
 }
 
 void SImageCatalog::UpdateItem(const FImageCatalogItemData& Item)
 {
-	if (const TPair<const TSharedPtr<FImageCatalogItemData>, bool> *const ExistingItem = FindItemData(Item.Guid))
+	if (const TSharedPtr<FImageCatalogItemData>* const ItemPtr = Model->GetItem(Item.Guid).Get<0>())
 	{
-		*ExistingItem->Get<0>() = Item;
+		**ItemPtr = Item;
 	}
 }
 
 void SImageCatalog::UpdateItemInfo(const FGuid& Guid, const FText& Info)
 {
-	if (const TPair<const TSharedPtr<FImageCatalogItemData>, bool> *const ExistingItem = FindItemData(Guid))
+	if (const TSharedPtr<FImageCatalogItemData>* const ItemPtr = Model->GetItem(Guid).Get<0>())
 	{
-		ExistingItem->Get<0>()->Info = Info;
+		(*ItemPtr)->Info = Info;
 	}
 }
 
 void SImageCatalog::UpdateItemName(const FGuid& Guid, const FText& Name)
 {
-	if (const TPair<const TSharedPtr<FImageCatalogItemData>, bool> *const ExistingItem = FindItemData(Guid))
+	if (const TSharedPtr<FImageCatalogItemData>* const ItemPtr = Model->GetItem(Guid).Get<0>())
 	{
-		ExistingItem->Get<0>()->Name = Name;
+		(*ItemPtr)->Name = Name;
 	}
 }
 
 void SImageCatalog::UpdateItemThumbnail(const FGuid& Guid, const FSlateBrush& Thumbnail)
 {
-	if (const TPair<const TSharedPtr<FImageCatalogItemData>, bool> *const ExistingItem = FindItemData(Guid))
+	if (const TSharedPtr<FImageCatalogItemData>* const ItemPtr = Model->GetItem(Guid).Get<0>())
 	{
-		ExistingItem->Get<0>()->Thumbnail = Thumbnail;
+		(*ItemPtr)->Thumbnail = Thumbnail;
 	}
 }
 
 void SImageCatalog::UpdateItemToolTip(const FGuid& Guid, const FText& ToolTip)
 {
-	if (const TPair<const TSharedPtr<FImageCatalogItemData>, bool> *const ExistingItem = FindItemData(Guid))
+	if (const TSharedPtr<FImageCatalogItemData>* const ItemPtr = Model->GetItem(Guid).Get<0>())
 	{
-		ExistingItem->Get<0>()->ToolTip = ToolTip;
+		(*ItemPtr)->ToolTip = ToolTip;
 	}
 }
 }
diff --git a/Engine/Plugins/Experimental/ImageWidgets/Source/ImageWidgets/Private/SImageViewport.cpp b/Engine/Plugins/Experimental/ImageWidgets/Source/ImageWidgets/Private/SImageViewport.cpp
--- a/Engine/Plugins/Experimental/ImageWidgets/Source/ImageWidgets/Private/SImageViewport.cpp
+++ b/Engine/Plugins/Experimental/ImageWidgets/Source/ImageWidgets/Private/SImageViewport.cpp
@@ -238,6 +238,15 @@ namespace UE::ImageWidgets
 		return {PixelCoordsValid, PixelCoords};
 	}
 
+	void SImageViewport::RequestRedraw() const
+	{
+		const TSharedPtr<FEditorViewportClient> ViewportClient = GetViewportClient();
+		if (ViewportClient.IsValid())
+		{
+			ViewportClient->RedrawRequested(ViewportClient->Viewport);
+		}
+	}
+
 	TSharedRef<FEditorViewportClient> SImageViewport::MakeEditorViewportClient()
 	{
 		if (!Client.IsValid())
diff --git a/Engine/Plugins/Experimental/ImageWidgets/Source/ImageWidgets/Public/SImageCatalog.h b/Engine/Plugins/Experimental/ImageWidgets/Source/ImageWidgets/Public/SImageCatalog.h
--- a/Engine/Plugins/Experimental/ImageWidgets/Source/ImageWidgets/Public/SImageCatalog.h
+++ b/Engine/Plugins/Experimental/ImageWidgets/Source/ImageWidgets/Public/SImageCatalog.h
@@ -48,18 +48,31 @@ namespace UE::ImageWidgets
 		 */
 		DECLARE_DELEGATE_OneParam(FOnItemSelected, const FGuid&)
 
+		/**
+		 * Delegate that gets called for creating a context menu for a set of selected items.
+		 * Return @see SWidget::NullWidget to not show a context menu.
+		 */
+		DECLARE_DELEGATE_RetVal_OneParam(TSharedPtr<SWidget>, FOnGetContextMenu, const TArray<FGuid>&)
+
 		SLATE_BEGIN_ARGS(SImageCatalog)
+				: _SelectionMode(ESelectionMode::Multi)
 			{
 			}
 
-			/** Header text for regular items */
-			SLATE_ARGUMENT(FText, ItemsHeading)
+			/** Header text for regular items. */
+			SLATE_ATTRIBUTE(FText, ItemsHeading)
 
-			/** Header text for pinned items */
-			SLATE_ARGUMENT(FText, PinnedItemsHeading)
+			/** Header text for pinned items. */
+			SLATE_ATTRIBUTE(FText, PinnedItemsHeading)
+
+			/** Defines the selection behavior within an item list, e.g. only allow single item selection or do not allow any selection. */
+			SLATE_ARGUMENT(ESelectionMode::Type, SelectionMode)
 
 			/** Delegate that gets called when an item is selected in the catalog. */
 			SLATE_EVENT(FOnItemSelected, OnItemSelected)
+
+			/** Delegate that gets called for creating a context menu for a set of selected items. */
+			SLATE_EVENT(FOnGetContextMenu, OnGetContextMenu)
 		SLATE_END_ARGS()
 
 		/**
@@ -72,14 +85,80 @@ namespace UE::ImageWidgets
 		 * Adds a regular item to the catalog.
 		 * @param Item Data for the item that is being added.
 		 */
-		IMAGEWIDGETS_API void AddItem(const TSharedPtr<FImageCatalogItemData>& Item);
+		IMAGEWIDGETS_API bool AddItem(const TSharedPtr<FImageCatalogItemData>& Item);
+
+		/**
+		 * Adds a regular item to the catalog right before another regular item.
+		 * If the unique identifier for the other item is invalid, the new item will be added at the end.
+		 * @param Item Data for the item that is being added.
+		 * @param BeforeItemWithThisGuid Unique identifier of the other item before which the new item should be added. 
+		 */
+		IMAGEWIDGETS_API bool AddItem(const TSharedPtr<FImageCatalogItemData>& Item, const FGuid& BeforeItemWithThisGuid);
 
 		/**
 		 * Adds a pinned item to the catalog. Pinned items appear in a separate list above regular items.
 		 * @param Item Data for the item that is being added. 
 		 */
-		IMAGEWIDGETS_API void AddPinnedItem(const TSharedPtr<FImageCatalogItemData>& Item);
+		IMAGEWIDGETS_API bool AddPinnedItem(const TSharedPtr<FImageCatalogItemData>& Item);
+
+		/**
+		 * Adds a pinned item to the catalog right before another pinned item.
+		 * If the unique identifier for the other item is invalid, the new item will be added at the end.
+		 * @param Item Data for the item that is being added.
+		 * @param BeforeItemWithThisGuid Unique identifier of the other item before which the new item should be added. 
+		 */
+		IMAGEWIDGETS_API bool AddPinnedItem(const TSharedPtr<FImageCatalogItemData>& Item, const FGuid& BeforeItemWithThisGuid);
+
+		/**
+		 * Remove an existing regular or pinned item.
+		 * Nothing happens if no item with the given unique identifier exists in the catalog.
+		 * @param Guid The unique identifier of the existing item.
+		 */
+		IMAGEWIDGETS_API bool RemoveItem(const FGuid& Guid);
+
+		/**
+		 * Retrieves the existing item for a given unique identifier.
+		 * @param Guid The unique identifier of the item.
+		 * @return The pointer to the item or an invalid pointer if no item with the given unique identifier exists.
+		 */
+		IMAGEWIDGETS_API TSharedPtr<const FImageCatalogItemData> GetItem(const FGuid& Guid) const;
+
+		/**
+		 * Returns if an existing item is a pinned or regular item as well as the index in the respective item list.
+		 * @param Guid The unique identifier of the item.
+		 * @return Tuple where the first value indicates if an item is pinned and the second value is the index in the respective item list or no tuple if no
+		 * item with the given unique identifier exists.
+		 */
+		IMAGEWIDGETS_API TOptional<TTuple<bool, int32>> GetItemIndex(const FGuid& Guid) const;
+
+		/**
+		 * Retrieves the existing regular item for a given index within the regular item list.
+		 * @param Index Index of the regular item; should be at least 0 and less than @see NumItems.
+		 * @return The pointer to the item or an invalid pointer if the index is invalid.
+		 */
+		IMAGEWIDGETS_API TSharedPtr<const FImageCatalogItemData> GetItemAt(int32 Index) const;
+
+		/**
+		 * Retrieves the existing pinned item for a given index within the pinned item list.
+		 * @param Index Index of the pinned item; should be at least 0 and less than @see NumPinnedItems.
+		 * @return The pointer to the item or an invalid pointer if the index is invalid.
+		 */
+		IMAGEWIDGETS_API TSharedPtr<const FImageCatalogItemData> GetPinnedItemAt(int32 Index) const;
+		
+		/**
+		 * Returns the unique identifier of the item at the given index within the regular item list.
+		 * @param Index Index of the regular item; should be at least 0 and less than @see NumItems.
+		 * @return Unique identifier of the item or no value if the index is invalid.
+		 */
+		IMAGEWIDGETS_API TOptional<FGuid> GetItemGuidAt(int32 Index) const;
 
+		/**
+		 * Returns the unique identifier of the item at the given index within the pinned item list.
+		 * @param Index Index of the pinned item; should be at least 0 and less than @see NumPinnedItems.
+		 * @return Unique identifier of the item or no value if the index is invalid.
+		 */
+		IMAGEWIDGETS_API TOptional<FGuid> GetPinnedItemGuidAt(int32 Index) const;
+		
 		/**
 		 * Returns the number of regular items in the catalog, i.e. items that are not pinned.
 		 */
@@ -95,6 +174,27 @@ namespace UE::ImageWidgets
 		 */
 		IMAGEWIDGETS_API int32 NumTotalItems() const;
 
+		/**
+		 * Check if an item is pinned.
+		 * @param Guid The unique identifier of the item.
+		 * @return True if the item exists and is pinned. 
+		 */
+		IMAGEWIDGETS_API bool ItemIsPinned(const FGuid& Guid) const;
+
+		/**
+		 * Pin an existing item.
+		 * @param Guid The unique identifier of the item.
+		 * @return True if the item exists and was not already pinned.
+		 */
+		IMAGEWIDGETS_API bool PinItem(const FGuid& Guid);
+
+		/**
+		 * Unpin an existing item.
+		 * @param Guid The unique identifier of the item.
+		 * @return True if the item exists and was not already unpinned.
+		 */
+		IMAGEWIDGETS_API bool UnpinItem(const FGuid& Guid);
+
 		/**
 		 * Select an existing regular or pinned item.
 		 * Nothing happens if no item with the given unique identifier exists in the catalog.
@@ -142,8 +242,8 @@ namespace UE::ImageWidgets
 		IMAGEWIDGETS_API void UpdateItemToolTip(const FGuid& Guid, const FText& ToolTip);
 
 	private:
-		/** Find an existing regular or pinned item's data. The second value is true if the item is pinned. */
-		const TPair<const TSharedPtr<FImageCatalogItemData>, bool>* FindItemData(const FGuid& Guid);
+		/** Collects the selected items and triggers the callback to create the context menu. */
+		TSharedPtr<SWidget> OnContextMenuOpening() const;
 
 		/** Widget for listing all regular items. */
 		TSharedPtr<SListView<TSharedPtr<FImageCatalogItemData>>> ItemsListView;
@@ -151,16 +251,19 @@ namespace UE::ImageWidgets
 		/** Widget for listing all pinned items. */
 		TSharedPtr<SListView<TSharedPtr<FImageCatalogItemData>>> PinnedItemsListView;
 
-		/** Regular item data. */
-		TArray<TSharedPtr<FImageCatalogItemData>> Items;
-
-		/** Regular item data. */
-		TArray<TSharedPtr<FImageCatalogItemData>> PinnedItems;
+		/** Header text for regular items. */
+		TAttribute<FText> ItemsHeading;
 
-		/** Mapping from the items' unique identifier to their respective data and if it is a pinned item. */
-		TMap<FGuid, TPair<const TSharedPtr<FImageCatalogItemData>, bool>> GuidToItemMapping;
+		/** Header text for pinned items. */
+		TAttribute<FText> PinnedItemsHeading;
 
 		/** Delegate that gets called when an item is selected. */
 		FOnItemSelected OnItemSelected;
+
+		/** Delegate that gets called to create a context menu for a set of selected items. */
+		FOnGetContextMenu OnGetContextMenu;
+
+		/** Internal item storage. */
+		TPimplPtr<class FItemModel> Model;
 	};
 }
diff --git a/Engine/Plugins/Experimental/ImageWidgets/Source/ImageWidgets/Public/SImageViewport.h b/Engine/Plugins/Experimental/ImageWidgets/Source/ImageWidgets/Public/SImageViewport.h
--- a/Engine/Plugins/Experimental/ImageWidgets/Source/ImageWidgets/Public/SImageViewport.h
+++ b/Engine/Plugins/Experimental/ImageWidgets/Source/ImageWidgets/Public/SImageViewport.h
@@ -207,6 +207,11 @@ namespace UE::ImageWidgets
 		 */
 		IMAGEWIDGETS_API FPixelCoordinatesUnderCursorResult GetPixelCoordinatesUnderCursor() const;
 
+		/**
+		 * Request to redraw the viewport contents as soon as possible.
+		 */
+		IMAGEWIDGETS_API void RequestRedraw() const;
+
 		/**
 		 * Resets the camera controller to default values.
 		 * @param ImageSize Size of the currently displayed image in pixels
