ImageWidgets: Add customizable catalog groups

diff --git a/Engine/Plugins/Experimental/ImageWidgets/Source/ImageWidgets/Private/ColorViewerSample/ColorViewerWidget.cpp b/Engine/Plugins/Experimental/ImageWidgets/Source/ImageWidgets/Private/ColorViewerSample/ColorViewerWidget.cpp
--- a/Engine/Plugins/Experimental/ImageWidgets/Source/ImageWidgets/Private/ColorViewerSample/ColorViewerWidget.cpp
+++ b/Engine/Plugins/Experimental/ImageWidgets/Source/ImageWidgets/Private/ColorViewerSample/ColorViewerWidget.cpp
@@ -7,8 +7,7 @@
 #include "ColorViewerCommands.h"
 #include "ColorViewerStyle.h"
 #include "Algo/AllOf.h"
-#include "Algo/NoneOf.h"
-#include "Algo/Sort.h"
+#include "Algo/AnyOf.h"
 #include "Framework/MultiBox/MultiBoxBuilder.h"
 
 #define LOCTEXT_NAMESPACE "ColorViewerWidget"
@@ -19,6 +18,9 @@ namespace UE::ImageWidgets::Sample
 	{
 		ColorViewer = MakeShared<FColorViewer>();
 
+		GroupColors = FName("Colors");
+		GroupFavorites = FName("Favorites");
+
 		BindCommands();
 
 		// Create toolbar extensions for a button to randomize the displayed color as well as the tone mapping controls.
@@ -36,11 +38,11 @@ namespace UE::ImageWidgets::Sample
 				SAssignNew(Splitter, SSplitter)
 					.PhysicalSplitterHandleSize(2.0f)
 					+ SSplitter::Slot()
-						.Value(0.0f)
+						.Value(0.2f)
 						[
 							SAssignNew(Catalog, ImageWidgets::SImageCatalog)
-								.PinnedItemsHeading(LOCTEXT("PinnedColors", "Pinned Colors"))
-								.ItemsHeading(LOCTEXT("Colors", "Colors"))
+								.DefaultGroupName(GroupColors)
+								.DefaultGroupHeading(LOCTEXT("Colors", "Colors"))
 								.OnItemSelected_Lambda([&ColorViewer = ColorViewer, &Viewport = Viewport](const FGuid& ImageGuid)
 								{
 									if (ColorViewer)
@@ -55,7 +57,7 @@ namespace UE::ImageWidgets::Sample
 								.OnGetContextMenu(this, &SColorViewerWidget::GetContextMenu)
 						]
 					+ SSplitter::Slot()
-						.Value(1.0f)
+						.Value(0.8f)
 						[
 							SAssignNew(Viewport, ImageWidgets::SImageViewport, ColorViewer.ToSharedRef())
 								.ToolbarExtender(ToolbarExtender)
@@ -70,6 +72,8 @@ namespace UE::ImageWidgets::Sample
 								.bABComparisonEnabled(true)
 						]
 			];
+
+		Catalog->AddGroup(GroupFavorites, LOCTEXT("Favorites", "Favorites"), GroupColors);
 	}
 
 	FReply SColorViewerWidget::OnKeyDown(const FGeometry& Geometry, const FKeyEvent& KeyEvent)
@@ -133,13 +137,8 @@ namespace UE::ImageWidgets::Sample
 			const FLinearColor ToneMappedColor = ColorViewer->GetDefaultToneMappedColor(ColorItem->Color);
 
 			Catalog->AddItem(MakeShared<FImageCatalogItemData>(ColorItem->Guid, FSlateColorBrush(ToneMappedColor), Name, Info, ToolTip));
+			Catalog->ClearSelection();
 			Catalog->SelectItem(ColorItem->Guid);
-
-			if (bCatalogCollapsedOnInit && Catalog->NumTotalItems() > 1 && Splitter->SlotAt(0).GetSizeValue() <= 0.0f)
-			{
-				Splitter->SlotAt(0).SetSizeValue(0.2f);
-				bCatalogCollapsedOnInit = false;
-			}
 		}
 	}
 
@@ -154,50 +153,84 @@ namespace UE::ImageWidgets::Sample
 		}
 	}
 
+	bool SColorViewerWidget::RandomizeColorEnabled() const
+	{
+		const IImageViewer::FImageInfo Info = ColorViewer->GetCurrentImageInfo();
+		if (Info.bIsValid)
+		{
+			if (const TOptional<FName> GroupName = Catalog->GetItemGroupName(Info.Guid))
+			{
+				if (*GroupName == GroupColors)
+				{
+					return true;
+				}
+			}
+		}
+
+		return false;
+	}
+
 	TSharedPtr<SWidget> SColorViewerWidget::GetContextMenu(const TArray<FGuid>& Guids) const
 	{
 		checkSlow(!Guids.IsEmpty());
-		checkSlow(Algo::AllOf(Guids, [&Catalog = Catalog](const FGuid& Guid) { return Catalog->ItemIsPinned(Guid); })
-			|| Algo::NoneOf(Guids, [&Catalog = Catalog] (const FGuid& Guid) { return Catalog->ItemIsPinned(Guid); }));
+
+		auto IsFavorite = [&Catalog = Catalog, &GroupFavorites = GroupFavorites](const FGuid& Guid)
+		{
+			const TOptional<FName> GroupName = Catalog->GetItemGroupName(Guid);
+			return GroupName && *GroupName == GroupFavorites;
+		};
+
+		const bool bHaveFavorites = Algo::AnyOf(Guids, IsFavorite);
+		const bool bHaveNonFavorites = !Algo::AllOf(Guids, IsFavorite);
 
 		FMenuBuilder MenuBuilder(true, nullptr);
 
-		if (Catalog->ItemIsPinned(Guids[0]))
+		if (bHaveNonFavorites)
 		{
 			MenuBuilder.AddMenuEntry(
-				LOCTEXT("UnpinItems", "Unpin"),
-				LOCTEXT("UnpinItemsTooltip", "Removes the selected item(s) from the list of pinned items."),
-				FSlateIcon(FAppStyle::GetAppStyleSetName(), "Icons.Unpinned"),
-				FUIAction(FExecuteAction::CreateLambda(
-					[Guids, Catalog = Catalog]
-					{
-						if (Catalog)
+				LOCTEXT("AddFavorite", "Add To Favorites"),
+				LOCTEXT("AddFavoriteTooltip", "Adds the selected color(s) to the list of favorites."),
+				FSlateIcon(FAppStyle::GetAppStyleSetName(), "Icons.Pinned"),
+				FUIAction(
+					FExecuteAction::CreateLambda(
+						[Guids, Catalog = Catalog, GroupFavorites = GroupFavorites, GroupColors = GroupColors]
 						{
-							for (const FGuid& Guid : Guids)
+							if (Catalog)
 							{
-								Catalog->UnpinItem(Guid);
+								for (const FGuid& Guid : Guids)
+								{
+									const auto GroupName = Catalog->GetItemGroupName(Guid);
+									if (GroupName && *GroupName == GroupColors)
+									{
+										Catalog->MoveItem(Guid, GroupFavorites);
+									}
+									Catalog->SelectItem(Guid);
+								}
 							}
-							Catalog->SelectItem(Guids[0]);
-						}
-					})));
+						})));
 		}
-		else
+
+		if (bHaveFavorites)
 		{
 			MenuBuilder.AddMenuEntry(
-				LOCTEXT("PinItems", "Pin"),
-				LOCTEXT("PinItemsTooltip", "Adds the selected item(s) to the list of pinned items."),
-				FSlateIcon(FAppStyle::GetAppStyleSetName(), "Icons.Pinned"),
+				LOCTEXT("RemoveFavorite", "Remove from Favorites"),
+				LOCTEXT("RemoveFavoriteTooltip", "Removes the selected color(s) from the list of favorites."),
+				FSlateIcon(FAppStyle::GetAppStyleSetName(), "Icons.Unpinned"),
 				FUIAction(
 					FExecuteAction::CreateLambda(
-						[Guids, Catalog = Catalog]
+						[Guids, Catalog = Catalog, GroupFavorites = GroupFavorites, GroupColors = GroupColors]
 						{
 							if (Catalog)
 							{
 								for (const FGuid& Guid : Guids)
 								{
-									Catalog->PinItem(Guid);
+									const auto GroupName = Catalog->GetItemGroupName(Guid);
+									if (GroupName && *GroupName == GroupFavorites)
+									{
+										Catalog->MoveItem(Guid, GroupColors);
+									}
+									Catalog->SelectItem(Guid);
 								}
-								Catalog->SelectItem(Guids[0]);
 							}
 						})));
 		}
@@ -215,7 +248,7 @@ namespace UE::ImageWidgets::Sample
 		MenuBuilder.AddSeparator();
 
 		MenuBuilder.AddMenuEntry(
-			LOCTEXT("DeleteItems", "Delete"), LOCTEXT("DeleteItemsTooltip", "Deletes the selected item(s) from the catalog."),
+			LOCTEXT("DeleteColors", "Delete"), LOCTEXT("DeleteColorsTooltip", "Deletes the selected colors(s) from the catalog."),
 			FSlateIcon(FAppStyle::GetAppStyleSetName(), "GenericCommands.Delete"),
 			FUIAction(FExecuteAction::CreateLambda(
 				[Guids, Catalog = Catalog, ColorViewer = ColorViewer]
@@ -246,7 +279,8 @@ namespace UE::ImageWidgets::Sample
 
 		CommandList->MapAction(
 			Commands.RandomizeColor,
-			FExecuteAction::CreateSP(this, &SColorViewerWidget::RandomizeColor)
+			FExecuteAction::CreateSP(this, &SColorViewerWidget::RandomizeColor),
+			FCanExecuteAction::CreateSP(this, &SColorViewerWidget::RandomizeColorEnabled)
 		);
 
 		CommandList->MapAction(
diff --git a/Engine/Plugins/Experimental/ImageWidgets/Source/ImageWidgets/Private/ColorViewerSample/ColorViewerWidget.h b/Engine/Plugins/Experimental/ImageWidgets/Source/ImageWidgets/Private/ColorViewerSample/ColorViewerWidget.h
--- a/Engine/Plugins/Experimental/ImageWidgets/Source/ImageWidgets/Private/ColorViewerSample/ColorViewerWidget.h
+++ b/Engine/Plugins/Experimental/ImageWidgets/Source/ImageWidgets/Private/ColorViewerSample/ColorViewerWidget.h
@@ -43,6 +43,9 @@ namespace UE::ImageWidgets::Sample
 		/** Choose a random color for the current entry. */
 		void RandomizeColor();
 
+		/** Checks if the Randomize Color action is available. */
+		bool RandomizeColorEnabled() const;
+
 		/** Creates a context menu for a set of selected guids. */
 		TSharedPtr<SWidget> GetContextMenu(const TArray<FGuid>& Guids) const;
 
@@ -61,8 +64,9 @@ namespace UE::ImageWidgets::Sample
 		/** The commands used by this sample widget. */
 		TSharedPtr<FUICommandList> CommandList;
 
-		/** Indicates that the catalog is still collapsed. This gets set to false as soon as a catalog entry is added, or it is manually expanded. */
-		bool bCatalogCollapsedOnInit = true;
+		/** Names for groups in the catalog. */
+		FName GroupColors;
+		FName GroupFavorites;
 	};
 }
 
diff --git a/Engine/Plugins/Experimental/ImageWidgets/Source/ImageWidgets/Private/SImageCatalog.cpp b/Engine/Plugins/Experimental/ImageWidgets/Source/ImageWidgets/Private/SImageCatalog.cpp
--- a/Engine/Plugins/Experimental/ImageWidgets/Source/ImageWidgets/Private/SImageCatalog.cpp
+++ b/Engine/Plugins/Experimental/ImageWidgets/Source/ImageWidgets/Private/SImageCatalog.cpp
@@ -2,11 +2,13 @@
 
 #include "SImageCatalog.h"
 
-#include "ImageWidgetsLogCategory.h"
+#include <Algo/Accumulate.h>
+#include <Brushes/SlateColorBrush.h>
+#include <Styling/StyleColors.h>
+#include <Widgets/Layout/SSeparator.h>
+#include <Widgets/Views/SListView.h>
+
 #include "SImageCatalogItem.h"
-#include "Brushes/SlateColorBrush.h"
-#include "Styling/StyleColors.h"
-#include "Widgets/Layout/SSeparator.h"
 
 #define LOCTEXT_NAMESPACE "SImageViewerCatalog"
 
@@ -18,590 +20,874 @@ FImageCatalogItemData::FImageCatalogItemData(const FGuid Guid, const FSlateBrush
 }
 
 using FItemType = TSharedPtr<FImageCatalogItemData>;
-	
-class FItemModel
+
+struct FGroup
 {
-public:
+	FGroup(FName Name, FText Heading)
+		: Name(Name), Heading(MoveTemp(Heading))
+	{}
 
-	bool Add(const FItemType& Item, bool bIsPinned, const FGuid* BeforeThisGuid);
-	TTuple<bool, bool> Remove(const FGuid& Guid);
+	FName Name;
+	FText Heading;
+	TSharedPtr<SListView<FItemType>> ListView;
+	TArray<FItemType> Items;
+};
 
-	TTuple<const FItemType*, bool> GetItem(const FGuid& Guid) const;
-	const FItemType* GetItemAt(int32 Index, bool bIsPinned) const;
-	TOptional<TTuple<bool, int32>> GetItemIndex(const FGuid& Guid) const;
-	TOptional<FGuid> GetGuidAt(int32 Index, bool bIsPinned) const;
-	
-	bool IsPinned(const FGuid& Guid);
-	bool Pin(const FGuid& Guid);
-	bool Unpin(const FGuid& Guid);
+struct FItemLookup
+{
+	int32 GroupIndex;
+	int32 ItemIndex;
+};
+
+class FImpl
+{
+public:
+	FImpl(FName DefaultGroupName, const FText& DefaultGroupHeading, ESelectionMode::Type SelectionMode, bool bAllowSelectionAcrossGroups, bool bShowEmptyGroups,
+	      const SImageCatalog::FOnItemSelected& OnItemSelected, const SImageCatalog::FOnGetContextMenu& OnGetContextMenu,
+	      TSharedPtr<SVerticalBox> GroupsLayout);
+
+	FName GetDefaultGroupName() const { return DefaultGroupName; }
+	bool AddGroup(FName Name, const FText& Heading, const FName* BeforeGroupWithThisName);
+	TOptional<TArray<FGuid>> RemoveGroup(FName Name, const FName* GroupToMoveItemsInto);
+	int32 NumGroups() const;
+	bool SetGroupHeading(FName Name, const FText& Heading);
+	TOptional<FName> GetGroupNameAt(int32 Index) const;
+
+	bool AddItem(const FItemType& Item, const FName* GroupName, const FGuid* BeforeItemWithThisGuid);
+	bool MoveItem(const FGuid& Guid, const FName* GroupName, const FGuid* BeforeItemWithThisGuid);
+	bool RemoveItem(const FGuid& Guid);
+	TSharedPtr<const FImageCatalogItemData> GetItem(const FGuid& Guid) const;
+	TOptional<int32> GetItemIndex(const FGuid& Guid) const;
+	TOptional<TTuple<FName, int32>> GetItemGroupNameAndIndex(const FGuid& Guid) const;
+	TSharedPtr<const FImageCatalogItemData> GetItemAt(int32 Index, const FName* GroupName) const;
+	TOptional<FGuid> GetItemGuidAt(int32 Index, const FName* GroupName) const;
+	bool UpdateItem(const FGuid& Guid, const FSlateBrush* Thumbnail, const FText* Name, const FText* Info, const FText* ToolTip);
 	
-	const TArray<FItemType>& GetPinnedItems() const { return PinnedItems; }
-	const TArray<FItemType>& GetUnpinnedItems() const { return UnpinnedItems; }
-	bool HasPinnedItems() const { return PinnedItems.Num() > 0; }
-	bool HasUnpinnedItems() const { return UnpinnedItems.Num() > 0; }
-	int32 NumPinnedItems() const { return PinnedItems.Num(); }
-	int32 NumUnpinnedItems() const { return UnpinnedItems.Num(); }
+	bool SelectItem(const FGuid& Guid, bool bSelected);
+	bool ClearSelection(const FName* GroupName);
 
-	void SortSelection(TArray<FGuid>& Selection);
+	int32 NumItems(const FName* GroupName) const;
+	TOptional<FName> GetItemGroupName(const FGuid& Guid) const;
 
 private:
+	
+	bool AddGroup(FName Name, const FText& Heading, int32 Index);
 
-	struct FLookupData
+	const FGroup* FindGroup(FName Name) const;
+	FGroup* FindGroup(FName Name)
 	{
-		bool bIsPinned;	// Is this item pinned?
-		int32 Index;	// Index in the respective array.
-	};
+		return const_cast<FGroup*>(const_cast<const FImpl*>(this)->FindGroup(Name));
+	}
 
-	const FLookupData* FindLookupData(const FGuid& Guid) const;
-	FLookupData* FindLookupData(const FGuid& Guid)
+	const FItemLookup* FindLookup(const FGuid& Guid) const;
+	FItemLookup* FindLookup(const FGuid& Guid)
 	{
-		return const_cast<FLookupData*>(const_cast<const FItemModel*>(this)->FindLookupData(Guid));
+		return const_cast<FItemLookup*>(const_cast<const FImpl*>(this)->FindLookup(Guid));
 	}
 
-	const FItemType* GetItem(const FLookupData& LookupData) const;
+	TSharedPtr<SWidget> OnContextMenuOpening() const;
+	void SortSelection(TArray<FGuid>& Selection) const;
+	void UpdateGroupMapping(int32 StartingIndex);
+	void UpdateItemMapping(TArray<FItemType>& Items, int32 StartingIndex);
+
+	const FName DefaultGroupName;
+	const FText DefaultGroupHeading;
+	const ESelectionMode::Type SelectionMode;
+	const bool bAllowSelectionAcrossGroups;
+	const bool bShowEmptyGroups;
+	SImageCatalog::FOnItemSelected OnItemSelected;
+	SImageCatalog::FOnGetContextMenu OnGetContextMenu;
+	TSharedPtr<SVerticalBox> GroupsLayout;
+	TArray<TUniquePtr<FGroup>> Groups;
+	TMap<FName, int32> GroupMapping;
+	TMap<FGuid, FItemLookup> ItemMapping;
+};
 
-	void SwapBetweenPinnedAndUnpinned(FLookupData& LookupData);
-	void UpdateMappingIndices(const TArray<FItemType>& Container, int32 FirstIndex);
+FImpl::FImpl(const FName DefaultGroupName, const FText& DefaultGroupHeading, const ESelectionMode::Type SelectionMode, const bool bAllowSelectionAcrossGroups,
+             const bool bShowEmptyGroups, const SImageCatalog::FOnItemSelected& OnItemSelected, const SImageCatalog::FOnGetContextMenu& OnGetContextMenu,
+             TSharedPtr<SVerticalBox> GroupsLayout)
+	: DefaultGroupName(DefaultGroupName)
+	, DefaultGroupHeading(DefaultGroupHeading)
+	, SelectionMode(SelectionMode)
+	, bAllowSelectionAcrossGroups(bAllowSelectionAcrossGroups)
+	, bShowEmptyGroups(bShowEmptyGroups)
+	, OnItemSelected(OnItemSelected)
+	, OnGetContextMenu(OnGetContextMenu)
+	, GroupsLayout(MoveTemp(GroupsLayout))
+{
+	AddGroup(DefaultGroupName, DefaultGroupHeading, nullptr);
+}
 
-	TArray<FItemType> PinnedItems;
-	TArray<FItemType> UnpinnedItems;
-	TMap<FGuid, FLookupData> GuidToLookupDataMapping;
-};
+bool FImpl::AddGroup(const FName Name, const FText& Heading, const FName* BeforeGroupWithThisName)
+{
+	if (GroupsLayout.IsValid())
+	{
+		if (BeforeGroupWithThisName != nullptr)
+		{
+			if (const int32 *const Index = GroupMapping.Find(*BeforeGroupWithThisName))
+			{
+				return AddGroup(Name, Heading, *Index);
+			}
+		}
+
+		return AddGroup(Name, Heading, Groups.Num());
+	}
+	return false;
+}
+
+TOptional<TArray<FGuid>> FImpl::RemoveGroup(const FName Name, const FName* GroupToMoveItemsInto)
+{
+	if (Name == DefaultGroupName)
+	{
+		return {};
+	}
+
+	if (const int32* GroupIndex = GroupMapping.Find(Name))
+	{
+		check(0 <= *GroupIndex && *GroupIndex < Groups.Num());
+		FGroup& Group = *Groups[*GroupIndex];
+		TArray<FItemType>& GroupItems = Group.Items;
+
+		TArray<FGuid> AffectedGuids;
+		AffectedGuids.Reserve(GroupItems.Num());
+
+		const int32* NewGroupIndex = GroupToMoveItemsInto ? GroupMapping.Find(*GroupToMoveItemsInto) : nullptr;
+		if (NewGroupIndex)
+		{
+			check(0 <= *NewGroupIndex && *NewGroupIndex < Groups.Num());
+			TArray<FItemType>& NewGroupItems = Groups[*NewGroupIndex]->Items;
+
+			for (FItemType& Item : GroupItems)
+			{
+				AffectedGuids.Add(Item->Guid);
+
+				FItemLookup* const Lookup = FindLookup(Item->Guid);
+				check(Lookup);
+				Lookup->GroupIndex = *NewGroupIndex;
+				Lookup->ItemIndex = NewGroupItems.Num();
+
+				NewGroupItems.Emplace(MoveTemp(Item));
+			}
+		}
+		else
+		{
+			for (const FItemType& Item : GroupItems)
+			{
+				AffectedGuids.Add(Item->Guid);
+				ItemMapping.Remove(Item->Guid);
+			}
+		}
+
+		check(GroupsLayout->GetChildren() && *GroupIndex < GroupsLayout->NumSlots());
+		GroupsLayout->RemoveSlot(GroupsLayout->GetChildren()->GetChildAt(*GroupIndex));		
+		Groups.RemoveAt(*GroupIndex);
+		GroupMapping.Remove(Name);
+		UpdateGroupMapping(*GroupIndex);
+
+		return {AffectedGuids};
+	}
+	
+	return {};
+}
 
-bool FItemModel::Add(const FItemType& Item, const bool bIsPinned, const FGuid* BeforeThisGuid)
+int32 FImpl::NumGroups() const
 {
-	if (!GuidToLookupDataMapping.Contains(Item->Guid))
+	return Groups.Num(); 
+}
+
+bool FImpl::SetGroupHeading(const FName Name, const FText& Heading)
+{
+	if (FGroup *const Group = FindGroup(Name))
 	{
-		TArray<FItemType>& Container = bIsPinned ? PinnedItems : UnpinnedItems;
+		Group->Heading = Heading;
+		return true;
+	}
 
-		const int32 Index = [this, &Item, bIsPinned, BeforeThisGuid, &Container]
+	return false;
+}
+
+TOptional<FName> FImpl::GetGroupNameAt(const int32 Index) const
+{
+	if (0 <= Index && Index < Groups.Num())
+	{
+		return {Groups[Index]->Name};
+	}
+
+	return {};
+}
+
+bool FImpl::AddItem(const FItemType& Item, const FName* GroupName, const FGuid* BeforeItemWithThisGuid)
+{
+	if (GroupName == nullptr)
+	{
+		GroupName = &DefaultGroupName;
+	}
+
+	if (const int32* GroupIndex = GroupMapping.Find(*GroupName))
+	{
+		check(0 <= *GroupIndex && *GroupIndex < Groups.Num());
+		FGroup& Group = *Groups[*GroupIndex];
+
+		const int32 ItemIndex = [&Item, BeforeItemWithThisGuid, &Group, this]
 		{
-			if (BeforeThisGuid)
+			if (BeforeItemWithThisGuid)
 			{
-				if (FLookupData *const LookupData = FindLookupData(*BeforeThisGuid))
+				if (FItemLookup* LookupData = FindLookup(*BeforeItemWithThisGuid))
 				{
-					if (LookupData->bIsPinned == bIsPinned)
+					if (Group.Name == Groups[LookupData->GroupIndex]->Name)
 					{
 						// Set index for added item.
-						const int32 NewIndex = LookupData->Index;
-
-						// Increase index for the item we push back.
-						++LookupData->Index;
+						const int32 Index = LookupData->ItemIndex;
 
 						// Add item at the new index.
-						Container.EmplaceAt(NewIndex, Item);
+						Group.Items.EmplaceAt(Index, Item);
+
+						// Increase index for the item we push back.
+						++LookupData->ItemIndex;
 
 						// Update all lookup data for items that come after the item we used to determine the insert location.
 						// This way we save the effort for finding the same lookup data again.
-						UpdateMappingIndices(Container, LookupData->Index + 1);
+						UpdateItemMapping(Group.Items, LookupData->ItemIndex + 1);
 
 						// Tell the outside where the new item was added.
-						return NewIndex;
+						return Index;
 					}
 				}
 			}
-			
-			return Container.Add(Item);
+
+			return Group.Items.Emplace(Item);
 		}();
 
-		GuidToLookupDataMapping.Add(Item->Guid, {bIsPinned, Index});
+		ItemMapping.Emplace(Item->Guid, {*GroupIndex, ItemIndex});
+
+		Group.ListView->RequestListRefresh();
+
 		return true;
 	}
 
 	return false;
 }
 
-TTuple<bool, bool> FItemModel::Remove(const FGuid& Guid)
+bool FImpl::MoveItem(const FGuid& Guid, const FName* GroupName, const FGuid* BeforeItemWithThisGuid)
 {
-	if (const FLookupData *const LookupData = FindLookupData(Guid))
+	// Do not try to move an item before itself.
+	if (BeforeItemWithThisGuid && *BeforeItemWithThisGuid == Guid)
+	{
+		return false;
+	}
+
+	if (FItemLookup *const Lookup = FindLookup(Guid))
 	{
-		const bool bIsPinned = LookupData->bIsPinned;
-		TArray<FItemType>& Container = bIsPinned ? PinnedItems : UnpinnedItems;
-		Container.RemoveAt(LookupData->Index, EAllowShrinking::No);
-		GuidToLookupDataMapping.Remove(Guid);
-		UpdateMappingIndices(Container, LookupData->Index);
-		return {true, bIsPinned};
+		FGroup& GroupFrom = *Groups[Lookup->GroupIndex];
+
+		const int32 *const GroupToIndex = GroupName == nullptr ? &Lookup->GroupIndex : GroupMapping.Find(*GroupName);
+		if (GroupToIndex)
+		{
+			check(0 <= *GroupToIndex && *GroupToIndex < Groups.Num());
+			FGroup& GroupTo = *Groups[*GroupToIndex];
+
+			if (BeforeItemWithThisGuid)
+			{
+				if (FItemLookup *const LookupBefore = FindLookup(*BeforeItemWithThisGuid))
+				{
+					// Only move if BeforeItems is in the correct group and if item isn't already in the correct position.
+					if (*GroupToIndex == LookupBefore->GroupIndex &&
+						(Lookup->GroupIndex != LookupBefore->GroupIndex || Lookup->ItemIndex + 1 != LookupBefore->ItemIndex))
+					{
+						FItemType Item = MoveTemp(GroupFrom.Items[Lookup->ItemIndex]);
+
+						GroupFrom.Items.RemoveAt(Lookup->ItemIndex, EAllowShrinking::No);
+						UpdateItemMapping(GroupFrom.Items, Lookup->ItemIndex);
+
+						GroupTo.Items.EmplaceAt(LookupBefore->ItemIndex, MoveTemp(Item));
+						Lookup->GroupIndex = LookupBefore->GroupIndex;
+						Lookup->ItemIndex = LookupBefore->ItemIndex;
+						++LookupBefore->ItemIndex;
+						UpdateItemMapping(GroupTo.Items, LookupBefore->ItemIndex + 1);
+
+						return true;
+					}
+				}
+			}
+			else
+			{
+				// Only move if item is not already in the correct group.
+				if (*GroupToIndex != Lookup->GroupIndex)
+				{
+					FItemType Item = MoveTemp(GroupFrom.Items[Lookup->ItemIndex]);
+				
+					GroupFrom.Items.RemoveAt(Lookup->ItemIndex, EAllowShrinking::No);
+					UpdateItemMapping(GroupFrom.Items, Lookup->ItemIndex);
+
+					Lookup->GroupIndex = *GroupToIndex;
+					Lookup->ItemIndex = GroupTo.Items.Emplace(MoveTemp(Item));
+
+					return true;
+				}
+			}
+		}
 	}
 
-	return {false, false};
+	return false;
 }
 
-TTuple<const FItemType*, bool> FItemModel::GetItem(const FGuid& Guid) const
+bool FImpl::RemoveItem(const FGuid& Guid)
 {
-	if (const FLookupData *const LookupData = FindLookupData(Guid))
+	if (const FItemLookup *const Lookup = FindLookup(Guid))
 	{
-		return {GetItem(*LookupData), LookupData->bIsPinned};
+		FGroup& Group = *Groups[Lookup->GroupIndex];
+		Group.Items.RemoveAt(Lookup->ItemIndex, EAllowShrinking::No);
+
+		ItemMapping.Remove(Guid);
+		UpdateItemMapping(Group.Items, Lookup->ItemIndex);
+
+		Group.ListView->RequestListRefresh();
+
+		return true;
 	}
 
-	return {nullptr, false};
+	return false;
 }
 
-const FItemType* FItemModel::GetItemAt(const int32 Index, const bool bIsPinned) const
+TSharedPtr<const FImageCatalogItemData> FImpl::GetItem(const FGuid& Guid) const
 {
-	const TArray<FItemType>& Container = bIsPinned ? PinnedItems : UnpinnedItems;
-	if (0 <= Index && Index < Container.Num())
+	if (const FItemLookup* Lookup = FindLookup(Guid))
 	{
-		return &Container[Index];
+		return Groups[Lookup->GroupIndex]->Items[Lookup->ItemIndex];
 	}
 
-	return nullptr;
+	return {};
 }
 
-TOptional<TTuple<bool, int32>> FItemModel::GetItemIndex(const FGuid& Guid) const
+TOptional<int32> FImpl::GetItemIndex(const FGuid& Guid) const
 {
-	if (const FLookupData *const LookupData = FindLookupData(Guid))
+	if (const FItemLookup* Lookup = FindLookup(Guid))
 	{
-		return {{LookupData->bIsPinned, LookupData->Index}};
+		return Lookup->ItemIndex;
 	}
 
 	return {};
 }
 
-TOptional<FGuid> FItemModel::GetGuidAt(const int32 Index, const bool bIsPinned) const
+TOptional<TTuple<FName, int32>> FImpl::GetItemGroupNameAndIndex(const FGuid& Guid) const
 {
-	const TArray<FItemType>& Container = bIsPinned ? PinnedItems : UnpinnedItems;
-	if (0 <= Index && Index < Container.Num())
+	if (const FItemLookup *const Lookup = FindLookup(Guid))
 	{
-		return Container[Index]->Guid;
+		return {MakeTuple(Groups[Lookup->GroupIndex]->Name, Lookup->ItemIndex)};
 	}
 
 	return {};
 }
 
-bool FItemModel::IsPinned(const FGuid& Guid)
+TSharedPtr<const FImageCatalogItemData> FImpl::GetItemAt(const int32 Index, const FName* GroupName) const
 {
-	const FLookupData *const LookupData = FindLookupData(Guid);
-	return LookupData && LookupData->bIsPinned;
+	if (!GroupName)
+	{
+		GroupName = &DefaultGroupName;
+	}
+
+	if (const FGroup* const Group = FindGroup(*GroupName))
+	{
+		if (0 <= Index && Index < Group->Items.Num())
+		{
+			return {Group->Items[Index]};
+		}
+	}
+	return {};
 }
 
-bool FItemModel::Pin(const FGuid& Guid)
+TOptional<FGuid> FImpl::GetItemGuidAt(int32 Index, const FName* GroupName) const
 {
-	FLookupData *const LookupData = FindLookupData(Guid);
-	if (LookupData && !LookupData->bIsPinned)
+	if (!GroupName)
 	{
-		SwapBetweenPinnedAndUnpinned(*LookupData);
-		return true;
+		GroupName = &DefaultGroupName;
 	}
 
-	return false;
+	if (const FGroup* const Group = FindGroup(*GroupName))
+	{
+		if (0 <= Index && Index < Group->Items.Num())
+		{
+			return {Group->Items[Index]->Guid};
+		}
+	}
+	return {};
 }
 
-bool FItemModel::Unpin(const FGuid& Guid)
+bool FImpl::UpdateItem(const FGuid& Guid, const FSlateBrush* Thumbnail, const FText* Name, const FText* Info, const FText* ToolTip)
 {
-	FLookupData *const LookupData = FindLookupData(Guid);
-	if (LookupData && LookupData->bIsPinned)
+	if (const FItemLookup* Lookup = FindLookup(Guid))
 	{
-		SwapBetweenPinnedAndUnpinned(*LookupData);
+		const FItemType& ItemPtr = Groups[Lookup->GroupIndex]->Items[Lookup->ItemIndex];
+		check(ItemPtr.IsValid());
+
+		FImageCatalogItemData& Item = *ItemPtr;
+
+		if (Thumbnail) Item.Thumbnail = *Thumbnail;
+		if (Name) Item.Name = *Name;
+		if (Info) Item.Info = *Info;
+		if (ToolTip) Item.ToolTip = *ToolTip;
+		
 		return true;
 	}
-
+	
 	return false;
 }
 
-void FItemModel::SortSelection(TArray<FGuid>& Selection)
+bool FImpl::SelectItem(const FGuid& Guid, const bool bSelected)
 {
-	Algo::Sort(Selection, [this, NumPinned = PinnedItems.Num()](const FGuid& A, const FGuid& B)
+	if (const FItemLookup *const Lookup = FindLookup(Guid))
 	{
-		const FLookupData *const LookupDataA = FindLookupData(A);
-		const FLookupData *const LookupDataB = FindLookupData(B);
+		const FGroup& Group = *Groups[Lookup->GroupIndex];
+		const FItemType& Item = Group.Items[Lookup->ItemIndex];
+		Group.ListView->SetItemSelection(Item, bSelected);
 
-		if (LookupDataA && !LookupDataB) return true;
-		if (!LookupDataA && LookupDataB) return false;
-		if (!LookupDataA && !LookupDataB) return true;
+		return true;
+	}
 
-		const int32 PositionA = (LookupDataA->bIsPinned ? 0 : NumPinned) + LookupDataA->Index;
-		const int32 PositionB = (LookupDataB->bIsPinned ? 0 : NumPinned) + LookupDataB->Index;
-		
-		return PositionA < PositionB;
-	});
+	return false;
 }
 
-const FItemModel::FLookupData* FItemModel::FindLookupData(const FGuid& Guid) const
+bool FImpl::ClearSelection(const FName* GroupName)
 {
-	if (const FLookupData* LookupData = GuidToLookupDataMapping.Find(Guid))
+	if (GroupName)
 	{
-		return LookupData;
+		if (const FGroup *const Group = FindGroup(*GroupName))
+		{
+			Group->ListView->ClearSelection();
+			return true;
+		}
+
+		return false;
 	}
 
-	UE_LOG(LogImageWidgets, Warning, TEXT("Cannot find catalog item for guid '%s'."), *Guid.ToString());
-	return nullptr;
-}
+	for (const TUniquePtr<FGroup>& Group : Groups)
+	{
+		Group->ListView->ClearSelection();
+	}
 
-const FItemType* FItemModel::GetItem(const FLookupData& LookupData) const
-{
-	const TArray<FItemType> *const Container = LookupData.bIsPinned ? &PinnedItems : &UnpinnedItems;
-	check(0 <= LookupData.Index && LookupData.Index < Container->Num());
-	return &(*Container)[LookupData.Index];
+	return true;
 }
 
-void FItemModel::SwapBetweenPinnedAndUnpinned(FLookupData& LookupData)
+int32 FImpl::NumItems(const FName* GroupName) const
 {
-	TArray<FItemType>& From = LookupData.bIsPinned ? PinnedItems : UnpinnedItems;
-	TArray<FItemType>& To = LookupData.bIsPinned ? UnpinnedItems : PinnedItems;
-	
-	const int32 NewIndex = To.Emplace(static_cast<FItemType&&>(From[LookupData.Index]));
-	From.RemoveAt(LookupData.Index, EAllowShrinking::No);
+	if (GroupName)
+	{
+		if (const FGroup* Group = FindGroup(*GroupName))
+		{
+			return Group->Items.Num();
+		}
 
-	UpdateMappingIndices(From, LookupData.Index);
+		return 0;
+	}
 
-	LookupData.Index = NewIndex;
-	LookupData.bIsPinned = !LookupData.bIsPinned;
+	return Algo::Accumulate(Groups, 0, [](int32 Num, const TUniquePtr<FGroup>& Group) { return Num + Group->Items.Num(); });
 }
 
-void FItemModel::UpdateMappingIndices(const TArray<FItemType>& Container, const int32 FirstIndex)
+TOptional<FName> FImpl::GetItemGroupName(const FGuid& Guid) const
 {
-	for (int32 Index = FirstIndex, Num = Container.Num(); Index < Num; ++Index)
+	if (const FItemLookup* Lookup = FindLookup(Guid))
 	{
-		FLookupData *const LookupData = GuidToLookupDataMapping.Find(Container[Index]->Guid);
-		check(LookupData);
-		LookupData->Index = Index;
+		return { Groups[Lookup->GroupIndex]->Name };
 	}
+
+	return {};
 }
 
-void SImageCatalog::Construct(const FArguments& InArgs)
+TSharedPtr<SWidget> FImpl::OnContextMenuOpening() const
 {
-	Model = MakePimpl<FItemModel>();
-	
-	ItemsHeading = InArgs._ItemsHeading;
-	PinnedItemsHeading = InArgs._PinnedItemsHeading;
-	OnItemSelected = InArgs._OnItemSelected;
-	OnGetContextMenu = InArgs._OnGetContextMenu;
+	if (!OnGetContextMenu.IsBound())
+	{
+		return SNullWidget::NullWidget;
+	}
 
-	const auto GenerateItemRow = [](const FItemType& ItemData, const TSharedRef<STableViewBase>& OwnerTable)
+	TArray<FGuid> SelectedGuids;
+	for (const TUniquePtr<FGroup>& Group : Groups)
 	{
-		static const FTableRowStyle TableRowStyle = []
-		{
-			FTableRowStyle Style = FAppStyle::Get().GetWidgetStyle<FTableRowStyle>("TableView.Row");
-			Style.SetOddRowBackgroundBrush(FSlateColorBrush(FStyleColors::Background));
-			Style.SetOddRowBackgroundHoveredBrush(FSlateColorBrush(FStyleColors::SelectHover));
-			Style.SetEvenRowBackgroundBrush(FSlateColorBrush(FStyleColors::Recessed));
-			Style.SetEvenRowBackgroundHoveredBrush(FSlateColorBrush(FStyleColors::SelectHover));
-			return Style;
-		}();
+		Algo::Transform(Group->ListView->GetSelectedItems(), SelectedGuids, [](const FItemType& Item) { return Item->Guid; });
+	}
 
-		TSharedPtr<SImageCatalogItem> ItemWidget;
-		SAssignNew(ItemWidget, SImageCatalogItem, ItemData);
+	if (SelectedGuids.IsEmpty())
+	{
+		return SNullWidget::NullWidget;
+	}
 
-		return SNew(STableRow<FItemType>, OwnerTable)
-			.Style(&TableRowStyle)
-			.ShowSelection(true)
-			[
-				ItemWidget.ToSharedRef()
-			];
-	};
+	SortSelection(SelectedGuids);
 
-	const auto GetPinnedItemsVisibility = [&Model = Model]
-	{
-		return Model->HasPinnedItems() ? EVisibility::Visible : EVisibility::Collapsed;
-	};
+	return OnGetContextMenu.Execute(SelectedGuids);
+}
 
-	const auto GetItemsVisibility = [&Model = Model]
-	{
-		return Model->HasUnpinnedItems() ? EVisibility::Visible : EVisibility::Collapsed;
-	};
+bool FImpl::AddGroup(FName Name, const FText& Heading, const int32 Index)
+{
+	check(0 <= Index && Index <= Groups.Num());
 
-	const auto PinnedItemsSelectionChanged = [this](const FItemType& Item, ESelectInfo::Type SelectInfo)
+	if (!GroupMapping.Contains(Name))
 	{
-		// Note that Item might be a nullptr since this callback is also executed when clearing the selection of the list.
-		if (Item.IsValid())
+		FGroup& Group = *Groups.EmplaceAt_GetRef(Index, MakeUnique<FGroup>(Name, Heading));
+		GroupMapping.Emplace(Name, Index);
+		UpdateGroupMapping(Index + 1);
+
+		const auto GetHeading = [&Group]
 		{
-			ItemsListView->ClearSelection();
-			OnItemSelected.ExecuteIfBound(Item->Guid);
-		}
-	};
+			return Group.Heading;
+		};
+		
+		const auto GetHeadingVisibility = [bShowEmptyGroups = bShowEmptyGroups, &Group]
+		{
+			return !Group.Heading.IsEmpty() && (bShowEmptyGroups || !Group.Items.IsEmpty()) ? EVisibility::Visible : EVisibility::Collapsed;
+		};
+		
+		const auto GetVisibility = [bShowEmptyGroups = bShowEmptyGroups, &Group]
+		{
+			return bShowEmptyGroups || !Group.Items.IsEmpty() ? EVisibility::Visible : EVisibility::Collapsed;
+		};
 
-	const auto ItemsSelectionChanged = [this](const FItemType& Item, ESelectInfo::Type SelectInfo)
-	{
-		// Note that Item might be a nullptr since this callback is also executed when clearing the selection of the list.
-		if (Item.IsValid())
+		const auto GenerateItemRow = [](const FItemType& ItemData, const TSharedRef<STableViewBase>& OwnerTable)
 		{
-			PinnedItemsListView->ClearSelection();
-			OnItemSelected.ExecuteIfBound(Item->Guid);
-		}
-	};
-
-	SAssignNew(PinnedItemsListView, SListView<FItemType>)
-					.ListItemsSource(&Model->GetPinnedItems())
-					.OnContextMenuOpening(this, &SImageCatalog::OnContextMenuOpening)
-					.OnGenerateRow_Lambda(GenerateItemRow)
-					.SelectionMode(InArgs._SelectionMode)
-					.OnSelectionChanged_Lambda(PinnedItemsSelectionChanged)
-					.ClearSelectionOnClick(false)
-					.Visibility_Lambda(GetPinnedItemsVisibility);
-
-	SAssignNew(ItemsListView, SListView<FItemType>)
-					.ListItemsSource(&Model->GetUnpinnedItems())
-					.OnContextMenuOpening(this, &SImageCatalog::OnContextMenuOpening)
-					.OnGenerateRow_Lambda(GenerateItemRow)
-					.SelectionMode(InArgs._SelectionMode)
-					.OnSelectionChanged_Lambda(ItemsSelectionChanged)
-					.ClearSelectionOnClick(false)
-					.ScrollbarVisibility(EVisibility::Visible);
-	
-	ChildSlot
-	[
-		SNew(SVerticalBox)
-		+ SVerticalBox::Slot()
-			.AutoHeight()
-			.Padding(2.0f, 4.0f, 2.0f, 4.0f)
-			[
-				SNew(STextBlock)
-					.Text_Lambda([&PinnedItemsHeading = PinnedItemsHeading] { return PinnedItemsHeading.Get({}); })
-					.Visibility_Lambda([&PinnedItemsHeading = PinnedItemsHeading, GetPinnedItemsVisibility]
+			static const FTableRowStyle TableRowStyle = []
+			{
+				FTableRowStyle Style = FAppStyle::Get().GetWidgetStyle<FTableRowStyle>("TableView.Row");
+				Style.SetOddRowBackgroundBrush(FSlateColorBrush(FStyleColors::Background));
+				Style.SetOddRowBackgroundHoveredBrush(FSlateColorBrush(FStyleColors::SelectHover));
+				Style.SetEvenRowBackgroundBrush(FSlateColorBrush(FStyleColors::Recessed));
+				Style.SetEvenRowBackgroundHoveredBrush(FSlateColorBrush(FStyleColors::SelectHover));
+				return Style;
+			}();
+
+			TSharedPtr<SImageCatalogItem> ItemWidget;
+			SAssignNew(ItemWidget, SImageCatalogItem, ItemData);
+
+			return SNew(STableRow<FItemType>, OwnerTable)
+				.Style(&TableRowStyle)
+				.ShowSelection(true)
+				[
+					ItemWidget.ToSharedRef()
+				];
+		};
+
+		const auto SelectionChanged = [bAllowSelectionAcrossGroups = bAllowSelectionAcrossGroups, &GroupName = Group.Name, &Groups = Groups, &OnItemSelected = OnItemSelected](const FItemType& Item, ESelectInfo::Type)
+		{
+			// Note that Item might be a nullptr since this callback is also executed when clearing the selection of the list.
+			if (Item.IsValid())
+			{
+				if (!bAllowSelectionAcrossGroups)
+				{
+					for (const TUniquePtr<FGroup>& Group : Groups)
 					{
-						return PinnedItemsHeading.Get({}).IsEmpty() ? EVisibility::Collapsed : GetPinnedItemsVisibility();
-					})
-					.TextStyle(FAppStyle::Get(), "DetailsView.CategoryTextStyle")
-			]
-		+ SVerticalBox::Slot()
-			.AutoHeight()
-			[
-				PinnedItemsListView.ToSharedRef()
-			]
-		+ SVerticalBox::Slot()
-			.AutoHeight()
-			[
-				SNew(SSeparator)
-					.Thickness(6.0f)
-					.Visibility_Lambda(GetPinnedItemsVisibility)
-			]
-		+ SVerticalBox::Slot()
+						if (Group->Name != GroupName)
+						{
+							if (Group->ListView)
+							{
+								Group->ListView->ClearSelection();
+							}
+						}
+					}
+				}
+
+				OnItemSelected.ExecuteIfBound(Item->Guid);
+			}
+		};
+
+		const auto OpenContextMenu = [this] { return OnContextMenuOpening(); };
+
+		GroupsLayout->InsertSlot(Index)
 			.AutoHeight()
-			.Padding(2.0f, 4.0f, 2.0f, 4.0f)
 			[
-				SNew(STextBlock)
-					.Text_Lambda([&ItemsHeading = ItemsHeading] { return ItemsHeading.Get({}); })
-					.Visibility_Lambda([&ItemsHeading = ItemsHeading, GetItemsVisibility]
-					{
-						return ItemsHeading.Get({}).IsEmpty() ? EVisibility::Collapsed : GetItemsVisibility();
-					})
-					.TextStyle(FAppStyle::Get(), "DetailsView.CategoryTextStyle")
-			]
-		+ SVerticalBox::Slot()
-			[
-				ItemsListView.ToSharedRef()
-			]
-	];
-}
+				SNew(SVerticalBox)
+					+ SVerticalBox::Slot()
+						.AutoHeight()
+						.Padding(2.0f, 4.0f, 2.0f, 4.0f)
+						[
+							SNew(STextBlock)
+								.Text_Lambda(GetHeading)
+								.Visibility_Lambda(GetHeadingVisibility)
+								.TextStyle(FAppStyle::Get(), "DetailsView.CategoryTextStyle")
+						]
+					+ SVerticalBox::Slot()
+						.AutoHeight()
+						.Padding(0.0f, 0.0f, 0.0f, 4.0f)
+						[
+							SAssignNew(Group.ListView, SListView<FItemType>)
+								.ListItemsSource(&Group.Items)
+								.Visibility_Lambda(GetVisibility)
+								.ScrollbarVisibility(EVisibility::Visible)
+                                .OnGenerateRow_Lambda(GenerateItemRow)
+								.SelectionMode(SelectionMode)
+								.ClearSelectionOnClick(false)
+								.OnSelectionChanged_Lambda(SelectionChanged)
+							    .OnContextMenuOpening_Lambda(OpenContextMenu)
+						]
+			];
 
-bool SImageCatalog::AddItem(const FItemType& Item)
-{
-	if (Model->Add(Item, false, nullptr))
-	{
-		ItemsListView->RequestListRefresh();
 		return true;
 	}
 
 	return false;
 }
 
-bool SImageCatalog::AddItem(const TSharedPtr<FImageCatalogItemData>& Item, const FGuid& BeforeItemWithThisGuid)
+const FGroup* FImpl::FindGroup(const FName Name) const
 {
-	if (Model->Add(Item, false, &BeforeItemWithThisGuid))
+	if (const int32* Index = GroupMapping.Find(Name))
 	{
-		ItemsListView->RequestListRefresh();
-		return true;
+		check(0 <= *Index && *Index < Groups.Num());
+		return Groups[*Index].Get();
 	}
-
-	return false;
+	return nullptr;
 }
 
-bool SImageCatalog::AddPinnedItem(const FItemType& Item)
+const FItemLookup* FImpl::FindLookup(const FGuid& Guid) const
 {
-	if (Model->Add(Item, true, nullptr))
+	if (const FItemLookup* Lookup = ItemMapping.Find(Guid))
 	{
-		PinnedItemsListView->RequestListRefresh();
-		return true;
+		check(0 <= Lookup->GroupIndex && Lookup->GroupIndex < Groups.Num());
+		check(0 <= Lookup->ItemIndex && Lookup->ItemIndex < Groups[Lookup->GroupIndex]->Items.Num());
+
+		return Lookup;
 	}
 
-	return false;
+	return nullptr;
 }
 
-bool SImageCatalog::AddPinnedItem(const TSharedPtr<FImageCatalogItemData>& Item, const FGuid& BeforeItemWithThisGuid)
+void FImpl::SortSelection(TArray<FGuid>& Selection) const
 {
-	if (Model->Add(Item, true, &BeforeItemWithThisGuid))
+	TArray<int32> GroupOffsets;
+	GroupOffsets.Reserve(Groups.Num());
+	int32 Offset = 0;
+	for (const TUniquePtr<FGroup>& Group : Groups)
 	{
-		PinnedItemsListView->RequestListRefresh();
-		return true;
+		GroupOffsets.Add(Offset);
+		Offset += Group->Items.Num();
 	}
+	
+	Algo::Sort(Selection, [&ItemMapping = ItemMapping, &GroupOffsets](const FGuid& A, const FGuid& B)
+	{
+		const FItemLookup *const LookupA = ItemMapping.Find(A);
+		const FItemLookup *const LookupB = ItemMapping.Find(B);
 
-	return false;
+		if (LookupA && !LookupB) return true;
+		if (!LookupA && LookupB) return false;
+		if (!LookupA && !LookupB) return true;
+
+		const int32 PositionA = GroupOffsets[LookupA->GroupIndex] + LookupA->ItemIndex;
+		const int32 PositionB = GroupOffsets[LookupB->GroupIndex] + LookupB->ItemIndex;
+		
+		return PositionA < PositionB;
+	});
 }
 
-bool SImageCatalog::RemoveItem(const FGuid& Guid)
+void FImpl::UpdateGroupMapping(const int32 StartingIndex)
 {
-	const auto [bSuccess, bIsPinned] = Model->Remove(Guid);
-	if (bSuccess)
+	for (int32 GroupIndex = StartingIndex, NumGroups = Groups.Num(); GroupIndex < NumGroups; ++GroupIndex)
 	{
-		const TSharedPtr<SListView<FItemType>>& ListView = bIsPinned ? PinnedItemsListView : ItemsListView;
-		ListView->RequestListRefresh();
-		return true;
-	}
+		int32 *const GroupMappingIndex = GroupMapping.Find(Groups[GroupIndex]->Name);
+		check(GroupMappingIndex);
+		*GroupMappingIndex = GroupIndex;
 
-	return false;
+		const TArray<FItemType>& Items = Groups[GroupIndex]->Items;
+		for (int32 ItemIndex = 0, NumItems = Items.Num(); ItemIndex < NumItems; ++ItemIndex)
+		{
+			FItemLookup *const Lookup = ItemMapping.Find(Items[ItemIndex]->Guid);
+			check(Lookup);
+			Lookup->GroupIndex = GroupIndex;
+		}
+	}
 }
 
-TSharedPtr<const FImageCatalogItemData> SImageCatalog::GetItem(const FGuid& Guid) const
+void FImpl::UpdateItemMapping(TArray<FItemType>& Items, const int32 StartingIndex)
 {
-	if (const FItemType* ItemPtr = Model->GetItem(Guid).Get<0>())
+	for (int32 Index = StartingIndex, Num = Items.Num(); Index < Num; ++Index)
 	{
-		return *ItemPtr;
+		FItemLookup *const Lookup = ItemMapping.Find(Items[Index]->Guid);
+		check(Lookup);
+		Lookup->ItemIndex = Index;
 	}
-	return {};
 }
 
-TOptional<TTuple<bool, int32>> SImageCatalog::GetItemIndex(const FGuid& Guid) const
+void SImageCatalog::Construct(const FArguments& Args)
 {
-	return Model->GetItemIndex(Guid);
+	TSharedPtr<SVerticalBox> GroupsLayout;
+
+	ChildSlot
+	[
+		SAssignNew(GroupsLayout, SVerticalBox)
+	];
+
+	Impl = MakePimpl<FImpl>(Args._DefaultGroupName, Args._DefaultGroupHeading, Args._SelectionMode, Args._bAllowSelectionAcrossGroups,
+	                        Args._bShowEmptyGroups, Args._OnItemSelected, Args._OnGetContextMenu, GroupsLayout);
 }
 
-TSharedPtr<const FImageCatalogItemData> SImageCatalog::GetItemAt(const int32 Index) const
+FName SImageCatalog::GetDefaultGroupName() const
 {
-	if (const FItemType* ItemPtr = Model->GetItemAt(Index, false))
-	{
-		return *ItemPtr;
-	}
-	return {};
+	return Impl->GetDefaultGroupName();
 }
 
-TSharedPtr<const FImageCatalogItemData> SImageCatalog::GetPinnedItemAt(const int32 Index) const
+bool SImageCatalog::AddGroup(const FName Name, const FText& Heading)
 {
-	if (const FItemType* ItemPtr = Model->GetItemAt(Index, true))
-	{
-		return *ItemPtr;
-	}
-	return {};
+	return Impl->AddGroup(Name, Heading, nullptr);
 }
 
-TOptional<FGuid> SImageCatalog::GetItemGuidAt(const int32 Index) const
+bool SImageCatalog::AddGroup(const FName Name, const FText& Heading, const FName BeforeGroupWithThisName)
 {
-	return Model->GetGuidAt(Index, false);
+	return Impl->AddGroup(Name, Heading, &BeforeGroupWithThisName);
 }
 
-TOptional<FGuid> SImageCatalog::GetPinnedItemGuidAt(const int32 Index) const 
+TOptional<TArray<FGuid>> SImageCatalog::RemoveGroup(const FName Name)
 {
-	return Model->GetGuidAt(Index, true);
+	return Impl->RemoveGroup(Name, nullptr);
 }
 
-int32 SImageCatalog::NumItems() const
+TOptional<TArray<FGuid>> SImageCatalog::RemoveGroup(const FName Name, const FName GroupToMoveItemsInto)
 {
-	return Model->NumUnpinnedItems();
+	return Impl->RemoveGroup(Name, &GroupToMoveItemsInto);
 }
 
-int32 SImageCatalog::NumPinnedItems() const
+int32 SImageCatalog::NumGroups() const
 {
-	return Model->NumPinnedItems();
+	return Impl->NumGroups();
 }
 
-int32 SImageCatalog::NumTotalItems() const
+TOptional<FName> SImageCatalog::GetGroupNameAt(int32 Index) const
 {
-	return NumItems() + NumPinnedItems();
+	return Impl->GetGroupNameAt(Index);
 }
 
-bool SImageCatalog::ItemIsPinned(const FGuid& Guid) const
+bool SImageCatalog::SetGroupHeading(const FName Name, const FText& Heading)
 {
-	return Model->IsPinned(Guid);
+	return Impl->SetGroupHeading(Name, Heading);
 }
-	
-bool SImageCatalog::PinItem(const FGuid& Guid)
+
+bool SImageCatalog::AddItem(const FItemType& Item)
 {
-	const bool bSuccess = Model->Pin(Guid);
-	if (bSuccess)
-	{
-		PinnedItemsListView->RequestListRefresh();
-		ItemsListView->RequestListRefresh();
-	}
-	return bSuccess;
+	return Impl->AddItem(Item, nullptr, nullptr);
 }
 
-bool SImageCatalog::UnpinItem(const FGuid& Guid)
+bool SImageCatalog::AddItem(const FItemType& Item, const FGuid& BeforeItemWithThisGuid)
 {
-	const bool bSuccess = Model->Unpin(Guid);
-	if (bSuccess)
-	{
-		PinnedItemsListView->RequestListRefresh();
-		ItemsListView->RequestListRefresh();
-	}
-	return bSuccess;
+	return Impl->AddItem(Item, nullptr, &BeforeItemWithThisGuid);
 }
 
-void SImageCatalog::SelectItem(const FGuid& Guid)
+bool SImageCatalog::AddItem(const TSharedPtr<FImageCatalogItemData>& Item, const FName Group)
 {
-	const auto [ItemPtr, bIsPinned] = Model->GetItem(Guid);
-	if (ItemPtr)
-	{
-		ItemsListView->ClearSelection();
-		PinnedItemsListView->ClearSelection();
+	return Impl->AddItem(Item, &Group, nullptr);
+}
 
-		SListView<FItemType> *const ListView = bIsPinned ? PinnedItemsListView.Get() : ItemsListView.Get();
-		ListView->SetItemSelection(*ItemPtr, true);
-	}
+bool SImageCatalog::AddItem(const TSharedPtr<FImageCatalogItemData>& Item, const FName Group, const FGuid& BeforeItemWithThisGuid)
+{
+	return Impl->AddItem(Item, &Group, &BeforeItemWithThisGuid);
 }
 
-TSharedPtr<SWidget> SImageCatalog::OnContextMenuOpening() const
+bool SImageCatalog::MoveItem(const FGuid& Guid, const FGuid& BeforeItemWithThisGuid)
 {
-	if (!OnGetContextMenu.IsBound())
-	{
-		return SNullWidget::NullWidget;
-	}
+	return Impl->MoveItem(Guid, nullptr, &BeforeItemWithThisGuid);
+}
 
-	TArray<FGuid> SelectedGuids = [&PinnedItemsListView = PinnedItemsListView, &ItemsListView = ItemsListView]
-	{
-		TArray<FGuid> Guids;
-		auto GetItemGuid = [](const FItemType& Item) { return Item->Guid; };
-		Algo::Transform(PinnedItemsListView->GetSelectedItems(), Guids, GetItemGuid);
-		Algo::Transform(ItemsListView->GetSelectedItems(), Guids, GetItemGuid);
-		return Guids;
-	}();
+bool SImageCatalog::MoveItem(const FGuid& Guid, FName Group)
+{
+	return Impl->MoveItem(Guid, &Group, nullptr);
+}
 
-	if (SelectedGuids.IsEmpty())
-	{
-		return SNullWidget::NullWidget;
-	}
+bool SImageCatalog::MoveItem(const FGuid& Guid, FName Group, const FGuid& BeforeItemWithThisGuid)
+{
+	return Impl->MoveItem(Guid, &Group, &BeforeItemWithThisGuid);
+}
 
-	Model->SortSelection(SelectedGuids);
+bool SImageCatalog::RemoveItem(const FGuid& Guid)
+{
+	return Impl->RemoveItem(Guid);
+}
 
-	return OnGetContextMenu.Execute(SelectedGuids);
+TSharedPtr<const FImageCatalogItemData> SImageCatalog::GetItem(const FGuid& Guid) const
+{
+	return Impl->GetItem(Guid);
 }
 
-void SImageCatalog::UpdateItem(const FImageCatalogItemData& Item)
+TOptional<int32> SImageCatalog::GetItemIndex(const FGuid& Guid) const
 {
-	if (const TSharedPtr<FImageCatalogItemData>* const ItemPtr = Model->GetItem(Item.Guid).Get<0>())
-	{
-		**ItemPtr = Item;
-	}
+	return Impl->GetItemIndex(Guid);
 }
 
-void SImageCatalog::UpdateItemInfo(const FGuid& Guid, const FText& Info)
+TOptional<TTuple<FName, int32>> SImageCatalog::GetItemGroupNameAndIndex(const FGuid& Guid) const
 {
-	if (const TSharedPtr<FImageCatalogItemData>* const ItemPtr = Model->GetItem(Guid).Get<0>())
-	{
-		(*ItemPtr)->Info = Info;
-	}
+	return Impl->GetItemGroupNameAndIndex(Guid);
 }
 
-void SImageCatalog::UpdateItemName(const FGuid& Guid, const FText& Name)
+TSharedPtr<const FImageCatalogItemData> SImageCatalog::GetItemAt(const int32 Index) const
 {
-	if (const TSharedPtr<FImageCatalogItemData>* const ItemPtr = Model->GetItem(Guid).Get<0>())
-	{
-		(*ItemPtr)->Name = Name;
-	}
+	return Impl->GetItemAt(Index, nullptr);
 }
 
-void SImageCatalog::UpdateItemThumbnail(const FGuid& Guid, const FSlateBrush& Thumbnail)
+TSharedPtr<const FImageCatalogItemData> SImageCatalog::GetItemAt(const int32 Index, const FName Group) const
 {
-	if (const TSharedPtr<FImageCatalogItemData>* const ItemPtr = Model->GetItem(Guid).Get<0>())
-	{
-		(*ItemPtr)->Thumbnail = Thumbnail;
-	}
+	return Impl->GetItemAt(Index, &Group);
 }
 
-void SImageCatalog::UpdateItemToolTip(const FGuid& Guid, const FText& ToolTip)
+TOptional<FGuid> SImageCatalog::GetItemGuidAt(const int32 Index) const
 {
-	if (const TSharedPtr<FImageCatalogItemData>* const ItemPtr = Model->GetItem(Guid).Get<0>())
-	{
-		(*ItemPtr)->ToolTip = ToolTip;
-	}
+	return Impl->GetItemGuidAt(Index, nullptr);
+}
+
+TOptional<FGuid> SImageCatalog::GetItemGuidAt(const int32 Index, const FName Group) const 
+{
+	return Impl->GetItemGuidAt(Index, &Group);
+}
+
+int32 SImageCatalog::NumItems() const
+{
+	return Impl->NumItems(nullptr);
+}
+
+int32 SImageCatalog::NumItems(FName Group) const
+{
+	return Impl->NumItems(&Group);
+}
+
+TOptional<FName> SImageCatalog::GetItemGroupName(const FGuid& Guid) const
+{
+	return Impl->GetItemGroupName(Guid);
+}
+
+bool SImageCatalog::SelectItem(const FGuid& Guid)
+{
+	return Impl->SelectItem(Guid, true);
+}
+
+bool SImageCatalog::DeselectItem(const FGuid& Guid)
+{
+	return Impl->SelectItem(Guid, false);
+}
+
+void SImageCatalog::ClearSelection()
+{
+	Impl->ClearSelection(nullptr);
+}
+
+bool SImageCatalog::ClearSelection(FName Group)
+{
+	return Impl->ClearSelection(&Group);
+}
+
+bool SImageCatalog::UpdateItem(const FImageCatalogItemData& Item)
+{
+	return Impl->UpdateItem(Item.Guid, &Item.Thumbnail, &Item.Name, &Item.Info, &Item.ToolTip);
+}
+
+bool SImageCatalog::UpdateItemInfo(const FGuid& Guid, const FText& Info)
+{
+	return Impl->UpdateItem(Guid, nullptr, nullptr, &Info, nullptr);
+}
+
+bool SImageCatalog::UpdateItemName(const FGuid& Guid, const FText& Name)
+{
+	return Impl->UpdateItem(Guid, nullptr, &Name, nullptr, nullptr);
+}
+
+bool SImageCatalog::UpdateItemThumbnail(const FGuid& Guid, const FSlateBrush& Thumbnail)
+{
+	return Impl->UpdateItem(Guid, &Thumbnail, nullptr, nullptr, nullptr);
+}
+
+bool SImageCatalog::UpdateItemToolTip(const FGuid& Guid, const FText& ToolTip)
+{
+	return Impl->UpdateItem(Guid, nullptr, nullptr, nullptr, &ToolTip);
 }
 }
 
diff --git a/Engine/Plugins/Experimental/ImageWidgets/Source/ImageWidgets/Public/IImageViewer.h b/Engine/Plugins/Experimental/ImageWidgets/Source/ImageWidgets/Public/IImageViewer.h
--- a/Engine/Plugins/Experimental/ImageWidgets/Source/ImageWidgets/Public/IImageViewer.h
+++ b/Engine/Plugins/Experimental/ImageWidgets/Source/ImageWidgets/Public/IImageViewer.h
@@ -2,8 +2,8 @@
 
 #pragma once
 
-#include "CoreMinimal.h"
-#include "Misc/TVariant.h"
+#include <CoreMinimal.h>
+#include <Misc/TVariant.h>
 
 class FCanvas;
 class FViewport;
diff --git a/Engine/Plugins/Experimental/ImageWidgets/Source/ImageWidgets/Public/SImageCatalog.h b/Engine/Plugins/Experimental/ImageWidgets/Source/ImageWidgets/Public/SImageCatalog.h
--- a/Engine/Plugins/Experimental/ImageWidgets/Source/ImageWidgets/Public/SImageCatalog.h
+++ b/Engine/Plugins/Experimental/ImageWidgets/Source/ImageWidgets/Public/SImageCatalog.h
@@ -2,15 +2,12 @@
 
 #pragma once
 
-#include "CoreMinimal.h"
-
+#include <CoreMinimal.h>
+#include <Framework/Views/ITypedTableView.h>
 #include <Widgets/SCompoundWidget.h>
-#include <Widgets/Views/SListView.h>
 
 namespace UE::ImageWidgets
 {
-	class SImageCatalogItem;
-
 	/**
 	 * Contains all data for a catalog item.
 	 */
@@ -35,9 +32,9 @@ namespace UE::ImageWidgets
 	};
 
 	/**
-	 * Generic catalog widget for listing and selecting 2D image-like content.
+	 * Generic catalog widget for listing and interacting with 2D image-like content.
 	 * Each catalog item is represented by its own widget based on its @see FImageCatalogItemData.
-	 * Entries in the catalog can be split into regular and pinned items, with pinned items being displayed at the top.
+	 * Entries in the catalog can be assigned to customizable groups.
 	 */
 	class SImageCatalog : public SCompoundWidget
 	{
@@ -55,215 +52,314 @@ namespace UE::ImageWidgets
 		DECLARE_DELEGATE_RetVal_OneParam(TSharedPtr<SWidget>, FOnGetContextMenu, const TArray<FGuid>&)
 
 		SLATE_BEGIN_ARGS(SImageCatalog)
-				: _SelectionMode(ESelectionMode::Multi)
+				: _DefaultGroupName(NAME_None)
+				, _DefaultGroupHeading(FText())
+				, _SelectionMode(ESelectionMode::Multi)
+				, _bAllowSelectionAcrossGroups(false)
+				, _bShowEmptyGroups(false)
 			{
 			}
 
-			/** Header text for regular items. */
-			SLATE_ATTRIBUTE(FText, ItemsHeading)
-
-			/** Header text for pinned items. */
-			SLATE_ATTRIBUTE(FText, PinnedItemsHeading)
-
-			/** Defines the selection behavior within an item list, e.g. only allow single item selection or do not allow any selection. */
-			SLATE_ARGUMENT(ESelectionMode::Type, SelectionMode)
+			/** Name of the default group, i.e. the group that gets used when no group is explicitly specified. */
+			SLATE_ARGUMENT(FName, DefaultGroupName)
+			
+			/** Header text for the default group. */
+			SLATE_ARGUMENT(FText, DefaultGroupHeading)
 
 			/** Delegate that gets called when an item is selected in the catalog. */
 			SLATE_EVENT(FOnItemSelected, OnItemSelected)
 
 			/** Delegate that gets called for creating a context menu for a set of selected items. */
 			SLATE_EVENT(FOnGetContextMenu, OnGetContextMenu)
+
+			/** Defines the selection behavior within an item list, e.g. only allow single item selection or do not allow any selection. */
+			SLATE_ARGUMENT(ESelectionMode::Type, SelectionMode)
+
+			/** When an item is selected, the selection in other groups will be cleared unless this flag is set to true. */
+			SLATE_ARGUMENT(bool, bAllowSelectionAcrossGroups)
+
+			/** Empty groups will be hidden unless this flag is set to true. */
+			SLATE_ARGUMENT(bool, bShowEmptyGroups)
 		SLATE_END_ARGS()
 
 		/**
 		 * Function used by Slate to construct the image catalog widget with the given arguments.
-		 * @param InArgs Slate arguments defined above
+		 * @param Args Slate arguments defined above
 		 */
-		IMAGEWIDGETS_API void Construct(const FArguments& InArgs);
+		IMAGEWIDGETS_API void Construct(const FArguments& Args);
 
 		/**
-		 * Adds a regular item to the catalog.
+		 * @returns The name of the default group.
+		 */
+		IMAGEWIDGETS_API FName GetDefaultGroupName() const;
+		
+		/**
+		 * Add a custom group to the catalog.
+		 * @param Name Unique identifier for the group.
+		 * @param Heading Header text to be used for the group; no header will be shown if this is empty.
+		 * @returns False if a group with the given name already exits, otherwise true.
+		 */
+		IMAGEWIDGETS_API bool AddGroup(FName Name, const FText& Heading);
+
+		/**
+		 * Add a custom group to the catalog before an already existing group.
+		 * @param Name Unique identifier for the group.
+		 * @param Heading Header text to be used for the group; no header will be shown if this is empty.
+		 * @param BeforeGroupWithThisName Name of an already existing group before which the new custom group is going to be added.
+		 * @returns False if a group with the given name already exits or if the group before which the new group is supposed to be added does not exist,
+		 * otherwise true.
+		 */
+		IMAGEWIDGETS_API bool AddGroup(FName Name, const FText& Heading, FName BeforeGroupWithThisName);
+
+		/**
+		 * Set the header text for an existing group.
+		 * @param Name Unique identifier for the group.
+		 * @param Heading Header text to be used for the group; no header will be shown if this is empty.
+		 * @returns False if a group with the given name does not exist, otherwise true.
+		 */
+		IMAGEWIDGETS_API bool SetGroupHeading(FName Name, const FText& Heading);
+
+		/**
+		* Remove an existing group and any items in the group.
+		* @param Name Unique identifier of the group; must not be the default group name.
+		* @returns An empty value if a group with the given name does not exist, otherwise a list of unique identifiers for the items that were in the group
+		* and got removed from the catalog.
+		*/
+		IMAGEWIDGETS_API TOptional<TArray<FGuid>> RemoveGroup(FName Name);
+
+		/**
+		* Remove an existing group, and moves any items in the group into another group.
+		* Note that if the other group does not exist, the items get moved into the default group instead.
+		* @param Name Unique identifier of the group; must not be the default group name.
+		* @param GroupToMoveItemsInto Name of an already existing group in which the items from the removed group get moved into.
+		* @returns An empty value if a group with the given name does not exist, otherwise a list of unique identifiers for the items that were in the group
+		* and got moved into the other group.
+		*/
+		IMAGEWIDGETS_API TOptional<TArray<FGuid>> RemoveGroup(FName Name, FName GroupToMoveItemsInto);
+
+		/**
+		 * Return the number of groups in the catalog including the default group.
+		 * @returns The number of groups in the catalog.
+		 */
+		IMAGEWIDGETS_API int32 NumGroups() const;
+
+		/**
+		 * Return the name of the group at a given position.
+		 * @param Index Position of the group, with 0 being the first, and @see NumGroups() - 1 being the last.
+		 * @returns An empty value if the given position is invalid, otherwise the name of the group at that position.
+		 */
+		IMAGEWIDGETS_API TOptional<FName> GetGroupNameAt(int32 Index) const;
+
+		/**
+		 * Add an item to the default group.
 		 * @param Item Data for the item that is being added.
+		 * @returns True if the item was added successfully.
 		 */
 		IMAGEWIDGETS_API bool AddItem(const TSharedPtr<FImageCatalogItemData>& Item);
 
 		/**
-		 * Adds a regular item to the catalog right before another regular item.
-		 * If the unique identifier for the other item is invalid, the new item will be added at the end.
+		 * Add an item to the default group right before an existing item.
+		 * If the unique identifier for the other item is invalid or not in the default group, the item will be added at the end of the default group instead.
 		 * @param Item Data for the item that is being added.
-		 * @param BeforeItemWithThisGuid Unique identifier of the other item before which the new item should be added. 
+		 * @param BeforeItemWithThisGuid Unique identifier of the other item before which the new item should be added.
+		 * @returns True if the item was added successfully.
 		 */
 		IMAGEWIDGETS_API bool AddItem(const TSharedPtr<FImageCatalogItemData>& Item, const FGuid& BeforeItemWithThisGuid);
 
 		/**
-		 * Adds a pinned item to the catalog. Pinned items appear in a separate list above regular items.
-		 * @param Item Data for the item that is being added. 
+		 * Add an item to an existing group.
+		 * @param Item Data for the item that should be added.
+		 * @param Group Name of the existing group to which the item should be added. 
+		 * @returns False if a group with the given name does not exist, otherwise true.
 		 */
-		IMAGEWIDGETS_API bool AddPinnedItem(const TSharedPtr<FImageCatalogItemData>& Item);
+		IMAGEWIDGETS_API bool AddItem(const TSharedPtr<FImageCatalogItemData>& Item, FName Group);
 
 		/**
-		 * Adds a pinned item to the catalog right before another pinned item.
-		 * If the unique identifier for the other item is invalid, the new item will be added at the end.
-		 * @param Item Data for the item that is being added.
-		 * @param BeforeItemWithThisGuid Unique identifier of the other item before which the new item should be added. 
+		 * Add an item to an existing group right before an existing item.
+		 * If the unique identifier for the other item is invalid or not in the given group, the item will be added at the end of the group instead.
+		 * @param Item Data for the item that should be added.
+		 * @param Group Name of the existing group to which the item should be added.
+		 * @param BeforeItemWithThisGuid Unique identifier of the other item before which the new item should be added.
+		 * @returns False if a group with the given name does not exist, otherwise true.
 		 */
-		IMAGEWIDGETS_API bool AddPinnedItem(const TSharedPtr<FImageCatalogItemData>& Item, const FGuid& BeforeItemWithThisGuid);
+		IMAGEWIDGETS_API bool AddItem(const TSharedPtr<FImageCatalogItemData>& Item, FName Group, const FGuid& BeforeItemWithThisGuid);
 
 		/**
-		 * Remove an existing regular or pinned item.
-		 * Nothing happens if no item with the given unique identifier exists in the catalog.
+		 * Move an already existing item before another item within the same group.
+		 * @param Guid The unique identifier of the existing item that should be moved.
+		 * @param BeforeItemWithThisGuid Unique identifier of the other item before which the item should be moved.
+		 * @returns False if the item to be moved or the item before which it should be moved does not exist or if both are not in the same group,
+		 * otherwise true.
+		 */
+		IMAGEWIDGETS_API bool MoveItem(const FGuid& Guid, const FGuid& BeforeItemWithThisGuid);
+
+		/**
+		 * Move an already existing item to another group.
+		 * @param Guid The unique identifier of the existing item that should be moved.
+		 * @param Group Name of the existing group to which the item should be moved.
+		 * @returns False if an item with the given unique identifier does not exist, a group with the given name does not exist, or the item is already in the
+		 * given group, otherwise true.
+		 */
+		IMAGEWIDGETS_API bool MoveItem(const FGuid& Guid, FName Group);
+
+		/**
+		 * Move an already existing item to another group, before another item within that group.
+		 * @param Guid The unique identifier of the existing item that should be moved.
+		 * @param Group Name of the existing group to which the item should be moved.
+		 * @param BeforeItemWithThisGuid Unique identifier of the other item before which the item should be moved.
+		 * @returns False if an item with the given unique identifier does not exist, a group with the given name does not exist, or the item is already in the
+		 * given group, otherwise true.
+		 */
+		IMAGEWIDGETS_API bool MoveItem(const FGuid& Guid, FName Group, const FGuid& BeforeItemWithThisGuid);
+
+		/**
+		 * Remove an existing item.
 		 * @param Guid The unique identifier of the existing item.
+		 * @returns False if an item with the given unique identifier does not exist, otherwise true.
 		 */
 		IMAGEWIDGETS_API bool RemoveItem(const FGuid& Guid);
 
 		/**
-		 * Retrieves the existing item for a given unique identifier.
+		 * Return the total number of items in the catalog.
+		 * @returns The number of items in the catalog across all groups.
+		 */
+		IMAGEWIDGETS_API int32 NumItems() const;
+
+		/**
+		 * Return the number of items in a group.
+		 * @param Group Name of an existing group.
+		 * @returns The number of items in the given group.
+		 */
+		IMAGEWIDGETS_API int32 NumItems(FName Group) const;
+
+		/**
+		 * Retrieve the existing item for a given unique identifier.
 		 * @param Guid The unique identifier of the item.
 		 * @return The pointer to the item or an invalid pointer if no item with the given unique identifier exists.
 		 */
 		IMAGEWIDGETS_API TSharedPtr<const FImageCatalogItemData> GetItem(const FGuid& Guid) const;
 
 		/**
-		 * Returns if an existing item is a pinned or regular item as well as the index in the respective item list.
-		 * @param Guid The unique identifier of the item.
-		 * @return Tuple where the first value indicates if an item is pinned and the second value is the index in the respective item list or no tuple if no
-		 * item with the given unique identifier exists.
+		 * Return the group an existing item belongs to.
+		 * @param Guid The unique identifier of the existing item.
+		 * @returns An empty value if an item with the given unique identifier does not exist, otherwise the name of the group the item belongs to.
 		 */
-		IMAGEWIDGETS_API TOptional<TTuple<bool, int32>> GetItemIndex(const FGuid& Guid) const;
+		IMAGEWIDGETS_API TOptional<FName> GetItemGroupName(const FGuid& Guid) const;
 
 		/**
-		 * Retrieves the existing regular item for a given index within the regular item list.
-		 * @param Index Index of the regular item; should be at least 0 and less than @see NumItems.
-		 * @return The pointer to the item or an invalid pointer if the index is invalid.
+		 * Return the index of an existing item within the group it belongs to.
+		 * @param Guid The unique identifier of the existing item.
+		 * @returns An empty value if no item with the given unique identifier exists, otherwise index of the item in the group.
 		 */
-		IMAGEWIDGETS_API TSharedPtr<const FImageCatalogItemData> GetItemAt(int32 Index) const;
+		IMAGEWIDGETS_API TOptional<int32> GetItemIndex(const FGuid& Guid) const;
 
 		/**
-		 * Retrieves the existing pinned item for a given index within the pinned item list.
-		 * @param Index Index of the pinned item; should be at least 0 and less than @see NumPinnedItems.
-		 * @return The pointer to the item or an invalid pointer if the index is invalid.
+		 * Return an existing item's group and the index within that group.
+		 * @param Guid The unique identifier of the existing item.
+		 * @returns An empty value if an item with the given unique identifier does not exist, otherwise the name of the group the item belongs to and
+		 * the index of the item in the group.
 		 */
-		IMAGEWIDGETS_API TSharedPtr<const FImageCatalogItemData> GetPinnedItemAt(int32 Index) const;
-		
+		IMAGEWIDGETS_API TOptional<TTuple<FName, int32>> GetItemGroupNameAndIndex(const FGuid& Guid) const;
+
 		/**
-		 * Returns the unique identifier of the item at the given index within the regular item list.
-		 * @param Index Index of the regular item; should be at least 0 and less than @see NumItems.
-		 * @return Unique identifier of the item or no value if the index is invalid.
+		 * Retrieve the existing item for a given index within the default group.
+		 * @param Index Index of the item in the default group; should be at least 0 and less than @see NumItems(@see GetDefaultGroup()).
+		 * @returns The pointer to the item or an invalid pointer if the index is invalid.
 		 */
-		IMAGEWIDGETS_API TOptional<FGuid> GetItemGuidAt(int32 Index) const;
+		IMAGEWIDGETS_API TSharedPtr<const FImageCatalogItemData> GetItemAt(int32 Index) const;
 
 		/**
-		 * Returns the unique identifier of the item at the given index within the pinned item list.
-		 * @param Index Index of the pinned item; should be at least 0 and less than @see NumPinnedItems.
-		 * @return Unique identifier of the item or no value if the index is invalid.
+		* Retrieve the existing item within a given group for a given index.
+		 * @param Index Index of the item in the given group; should be at least 0 and less than @see NumItems(Group).
+		 * @param Group Name of an existing group.
+		 * @returns The pointer to the item or an invalid pointer if the group does not exist or the index is invalid.
 		 */
-		IMAGEWIDGETS_API TOptional<FGuid> GetPinnedItemGuidAt(int32 Index) const;
+		IMAGEWIDGETS_API TSharedPtr<const FImageCatalogItemData> GetItemAt(int32 Index, FName Group) const;
 		
 		/**
-		 * Returns the number of regular items in the catalog, i.e. items that are not pinned.
-		 */
-		IMAGEWIDGETS_API int32 NumItems() const;
-
-		/**
-		 * Returns the number of pinned items in the catalog.
+		 * Retrieve the unique identifier of an existing item at the given index within the default group.
+		 * @param Index Index of the item in the default group; should be at least 0 and less than @see NumItems(@see GetDefaultGroup()).
+		 * @return Unique identifier of the item or no value if the index is invalid.
 		 */
-		IMAGEWIDGETS_API int32 NumPinnedItems() const;
+		IMAGEWIDGETS_API TOptional<FGuid> GetItemGuidAt(int32 Index) const;
 
 		/**
-		 * Returns the total number of items in the catalog, i.e. both regular items and pinned items.
+		 * Retrieve the unique identifier of an existing item at the given index within the given group.
+		 * @param Index Index of the item in the given group; should be at least 0 and less than @see NumItems(Group).
+		 * @param Group Name of an existing group.
+		 * @return Unique identifier of the item or no value if the group does not exist or the index is invalid.
 		 */
-		IMAGEWIDGETS_API int32 NumTotalItems() const;
-
+		IMAGEWIDGETS_API TOptional<FGuid> GetItemGuidAt(int32 Index, FName Group) const;
+		
 		/**
-		 * Check if an item is pinned.
-		 * @param Guid The unique identifier of the item.
-		 * @return True if the item exists and is pinned. 
+		 * Select an existing item.
+		 * @param Guid The unique identifier of the existing item.
+		 * @returns False if an item with the given unique identifier does not exist, otherwise true.
 		 */
-		IMAGEWIDGETS_API bool ItemIsPinned(const FGuid& Guid) const;
+		IMAGEWIDGETS_API bool SelectItem(const FGuid& Guid);
 
 		/**
-		 * Pin an existing item.
-		 * @param Guid The unique identifier of the item.
-		 * @return True if the item exists and was not already pinned.
+		 * Deselect an existing item.
+		 * @param Guid The unique identifier of the existing item.
+		 * @returns False if an item with the given unique identifier does not exist, otherwise true.
 		 */
-		IMAGEWIDGETS_API bool PinItem(const FGuid& Guid);
+		IMAGEWIDGETS_API bool DeselectItem(const FGuid& Guid);
 
 		/**
-		 * Unpin an existing item.
-		 * @param Guid The unique identifier of the item.
-		 * @return True if the item exists and was not already unpinned.
+		 * Clear any selection in the catalog, i.e. across all groups.
 		 */
-		IMAGEWIDGETS_API bool UnpinItem(const FGuid& Guid);
+		IMAGEWIDGETS_API void ClearSelection();
 
 		/**
-		 * Select an existing regular or pinned item.
-		 * Nothing happens if no item with the given unique identifier exists in the catalog.
-		 * @param Guid The unique identifier of the existing item.
+		 * Clear the selection for a given group; selections in other groups will be unchanged.
+		 * @param Group Name of an existing group.
+		 * @returns False if a group with the given name does not exist
 		 */
-		IMAGEWIDGETS_API void SelectItem(const FGuid& Guid);
+		IMAGEWIDGETS_API bool ClearSelection(FName Group);
 
 		/**
-		 * Update an existing regular or pinned item's data. The item data should contain the item's unique identifier.
-		 * Nothing happens if no item with the given unique identifier exists in the catalog.
-		 * @param Item Data for the item that is being updated, including the unique identifier of the existing item.
+		 * Update an existing item's data. The item data needs to contain the item's unique identifier.
+		 * @param Item Data for the exiting item that should be updated, including its unique identifier.
+		 * @returns False if an item with the given unique identifier does not exist, otherwise true.
 		 */
-		IMAGEWIDGETS_API void UpdateItem(const FImageCatalogItemData& Item);
+		IMAGEWIDGETS_API bool UpdateItem(const FImageCatalogItemData& Item);
 
 		/**
-		 * Update the info text of an existing regular or pinned item.
-		 * Nothing happens if no item with the given unique identifier exists in the catalog.
+		 * Update the info text of an existing item.
 		 * @param Guid The unique identifier of the existing item.
 		 * @param Info Text for the info label in the item's widget.
+		 * @returns False if an item with the given unique identifier does not exist, otherwise true.
 		 */
-		IMAGEWIDGETS_API void UpdateItemInfo(const FGuid& Guid, const FText& Info);
+		IMAGEWIDGETS_API bool UpdateItemInfo(const FGuid& Guid, const FText& Info);
 
 		/**
-		 * Update the info text of an existing regular or pinned item.
+		 * Update the info text of an existing item.
 		 * Nothing happens if no item with the given unique identifier exists in the catalog.
 		 * @param Guid The unique identifier of the existing item.
 		 * @param Name Text for the name label in the item's widget.
+		 * @returns False if an item with the given unique identifier does not exist, otherwise true.
 		 */
-		IMAGEWIDGETS_API void UpdateItemName(const FGuid& Guid, const FText& Name);
+		IMAGEWIDGETS_API bool UpdateItemName(const FGuid& Guid, const FText& Name);
 
 		/**
-		 * Update the thumbnail of an existing regular or pinned item.
+		 * Update the thumbnail of an existing item.
 		 * Nothing happens if no item with the given unique identifier exists in the catalog.
 		 * @param Guid The unique identifier of the existing item.
 		 * @param Thumbnail Brush used for the thumbnail in the item's widget.
+		 * @returns False if an item with the given unique identifier does not exist, otherwise true.
 		 */
-		IMAGEWIDGETS_API void UpdateItemThumbnail(const FGuid& Guid, const FSlateBrush& Thumbnail);
+		IMAGEWIDGETS_API bool UpdateItemThumbnail(const FGuid& Guid, const FSlateBrush& Thumbnail);
 
 		/**
-		 * Update the tooltip text of an existing regular or pinned item.
+		 * Update the tooltip text of an existing item.
 		 * Nothing happens if no item with the given unique identifier exists in the catalog.
 		 * @param Guid The unique identifier of the existing item.
 		 * @param ToolTip Text for the tooltip label in the item's widget.
+		 * @returns False if an item with the given unique identifier does not exist, otherwise true.
 		 */
-		IMAGEWIDGETS_API void UpdateItemToolTip(const FGuid& Guid, const FText& ToolTip);
+		IMAGEWIDGETS_API bool UpdateItemToolTip(const FGuid& Guid, const FText& ToolTip);
 
 	private:
-		/** Collects the selected items and triggers the callback to create the context menu. */
-		TSharedPtr<SWidget> OnContextMenuOpening() const;
-
-		/** Widget for listing all regular items. */
-		TSharedPtr<SListView<TSharedPtr<FImageCatalogItemData>>> ItemsListView;
-
-		/** Widget for listing all pinned items. */
-		TSharedPtr<SListView<TSharedPtr<FImageCatalogItemData>>> PinnedItemsListView;
-
-		/** Header text for regular items. */
-		TAttribute<FText> ItemsHeading;
-
-		/** Header text for pinned items. */
-		TAttribute<FText> PinnedItemsHeading;
-
-		/** Delegate that gets called when an item is selected. */
-		FOnItemSelected OnItemSelected;
-
-		/** Delegate that gets called to create a context menu for a set of selected items. */
-		FOnGetContextMenu OnGetContextMenu;
-
-		/** Internal item storage. */
-		TPimplPtr<class FItemModel> Model;
+		TPimplPtr<class FImpl> Impl;
 	};
 }
diff --git a/Engine/Plugins/Experimental/ImageWidgets/Source/ImageWidgets/Public/SImageViewport.h b/Engine/Plugins/Experimental/ImageWidgets/Source/ImageWidgets/Public/SImageViewport.h
--- a/Engine/Plugins/Experimental/ImageWidgets/Source/ImageWidgets/Public/SImageViewport.h
+++ b/Engine/Plugins/Experimental/ImageWidgets/Source/ImageWidgets/Public/SImageViewport.h
@@ -2,9 +2,8 @@
 
 #pragma once
 
-#include "SEditorViewport.h"
-
-#include "Framework/MultiBox/MultiBoxExtender.h"
+#include <SEditorViewport.h>
+#include <Framework/MultiBox/MultiBoxExtender.h>
 
 class SViewportToolBar;
 
