ImageWidgets: Fix Non-Unity compile error in SImageCatalogItem

diff --git a/Engine/Plugins/Experimental/ImageWidgets/Source/ImageWidgets/Private/SImageCatalogItem.cpp b/Engine/Plugins/Experimental/ImageWidgets/Source/ImageWidgets/Private/SImageCatalogItem.cpp
--- a/Engine/Plugins/Experimental/ImageWidgets/Source/ImageWidgets/Private/SImageCatalogItem.cpp
+++ b/Engine/Plugins/Experimental/ImageWidgets/Source/ImageWidgets/Private/SImageCatalogItem.cpp
@@ -3,6 +3,10 @@
 #include "SImageCatalogItem.h"
 #include "SImageCatalog.h"
 
+#include <Widgets/SBoxPanel.h>
+#include <Widgets/Images/SImage.h>
+#include <Widgets/Text/STextBlock.h>
+
 #define LOCTEXT_NAMESPACE "SImageViewerCatalogItem"
 
 namespace UE::ImageWidgets
