ImageWidgets:
* Add Catalog customizable group context menu
* Allow to collapse catalog groups
* Catalog layout and style improvements

diff --git a/Engine/Plugins/Experimental/ImageWidgets/Source/ImageWidgets/Private/ColorViewerSample/ColorViewerWidget.cpp b/Engine/Plugins/Experimental/ImageWidgets/Source/ImageWidgets/Private/ColorViewerSample/ColorViewerWidget.cpp
--- a/Engine/Plugins/Experimental/ImageWidgets/Source/ImageWidgets/Private/ColorViewerSample/ColorViewerWidget.cpp
+++ b/Engine/Plugins/Experimental/ImageWidgets/Source/ImageWidgets/Private/ColorViewerSample/ColorViewerWidget.cpp
@@ -54,7 +54,8 @@ namespace UE::ImageWidgets::Sample
 										Viewport->RequestRedraw();
 									}
 								})
-								.OnGetContextMenu(this, &SColorViewerWidget::GetContextMenu)
+								.OnGetGroupContextMenu(this, &SColorViewerWidget::GetGroupContextMenu)
+								.OnGetItemsContextMenu(this, &SColorViewerWidget::GetItemsContextMenu)
 						]
 					+ SSplitter::Slot()
 						.Value(0.8f)
@@ -170,7 +171,41 @@ namespace UE::ImageWidgets::Sample
 		return false;
 	}
 
-	TSharedPtr<SWidget> SColorViewerWidget::GetContextMenu(const TArray<FGuid>& Guids) const
+	TSharedPtr<SWidget> SColorViewerWidget::GetGroupContextMenu(FName GroupName) const
+	{
+		FMenuBuilder MenuBuilder(true, nullptr);
+
+		const int32 NumItems = Catalog->NumItems(GroupName);
+		if (NumItems > 0)
+		{
+			MenuBuilder.AddMenuEntry(
+				NumItems == 1
+					? FText::Format(LOCTEXT("DeleteGroupItem", "Delete single item in group"), NumItems)
+					: FText::Format(LOCTEXT("DeleteGroupItems", "Delete all {0} items in group"), NumItems),
+				LOCTEXT("DeleteGroupItems", "Deletes all items in this group."),
+				FSlateIcon(FAppStyle::GetAppStyleSetName(), "GenericCommands.Delete"),
+				FUIAction(FExecuteAction::CreateLambda(
+					[GroupName, Catalog = Catalog, ColorViewer = ColorViewer]
+					{
+						if (Catalog && ColorViewer)
+						{
+							for (int32 Index = Catalog->NumItems(GroupName) - 1; Index >= 0; --Index)
+							{
+								const TOptional<FGuid> Guid = Catalog->GetItemGuidAt(Index, GroupName);
+								if (Guid.IsSet())
+								{
+									Catalog->RemoveItem(Guid.GetValue());
+									ColorViewer->RemoveColor(Guid.GetValue());
+								}
+							}
+						}
+					})));
+		}
+
+		return MenuBuilder.MakeWidget();
+	}
+
+	TSharedPtr<SWidget> SColorViewerWidget::GetItemsContextMenu(const TArray<FGuid>& Guids) const
 	{
 		checkSlow(!Guids.IsEmpty());
 
@@ -248,7 +283,8 @@ namespace UE::ImageWidgets::Sample
 		MenuBuilder.AddSeparator();
 
 		MenuBuilder.AddMenuEntry(
-			LOCTEXT("DeleteColors", "Delete"), LOCTEXT("DeleteColorsTooltip", "Deletes the selected colors(s) from the catalog."),
+			LOCTEXT("DeleteColors", "Delete"),
+			LOCTEXT("DeleteColorsTooltip", "Deletes the selected colors(s) from the catalog."),
 			FSlateIcon(FAppStyle::GetAppStyleSetName(), "GenericCommands.Delete"),
 			FUIAction(FExecuteAction::CreateLambda(
 				[Guids, Catalog = Catalog, ColorViewer = ColorViewer]
diff --git a/Engine/Plugins/Experimental/ImageWidgets/Source/ImageWidgets/Private/ColorViewerSample/ColorViewerWidget.h b/Engine/Plugins/Experimental/ImageWidgets/Source/ImageWidgets/Private/ColorViewerSample/ColorViewerWidget.h
--- a/Engine/Plugins/Experimental/ImageWidgets/Source/ImageWidgets/Private/ColorViewerSample/ColorViewerWidget.h
+++ b/Engine/Plugins/Experimental/ImageWidgets/Source/ImageWidgets/Private/ColorViewerSample/ColorViewerWidget.h
@@ -8,6 +8,7 @@
 
 #include <SImageCatalog.h>
 #include <SImageViewport.h>
+#include <Widgets/Layout/SSplitter.h>
 
 namespace UE::ImageWidgets::Sample
 {
@@ -46,8 +47,11 @@ namespace UE::ImageWidgets::Sample
 		/** Checks if the Randomize Color action is available. */
 		bool RandomizeColorEnabled() const;
 
+		/** Creates a context menu for a group. */
+		TSharedPtr<SWidget> GetGroupContextMenu(FName GroupName) const;
+
 		/** Creates a context menu for a set of selected guids. */
-		TSharedPtr<SWidget> GetContextMenu(const TArray<FGuid>& Guids) const;
+		TSharedPtr<SWidget> GetItemsContextMenu(const TArray<FGuid>& Guids) const;
 
 		/** The image viewer implementation that contains the image data and renders the image. */
 		TSharedPtr<FColorViewer> ColorViewer;
@@ -57,7 +61,7 @@ namespace UE::ImageWidgets::Sample
 
 		/** The image catalog that holds all currently available images. */
 		TSharedPtr<SImageCatalog> Catalog;
-		
+
 		/** The image viewport in which the current image gets displayed in. */
 		TSharedPtr<SImageViewport> Viewport;
 
diff --git a/Engine/Plugins/Experimental/ImageWidgets/Source/ImageWidgets/Private/ImageWidgetsModule.cpp b/Engine/Plugins/Experimental/ImageWidgets/Source/ImageWidgets/Private/ImageWidgetsModule.cpp
--- a/Engine/Plugins/Experimental/ImageWidgets/Source/ImageWidgets/Private/ImageWidgetsModule.cpp
+++ b/Engine/Plugins/Experimental/ImageWidgets/Source/ImageWidgets/Private/ImageWidgetsModule.cpp
@@ -57,7 +57,7 @@ namespace UE::ImageWidgets
 		.SetDisplayName(LOCTEXT("ColorViewerTitle", "Color Viewer Sample"))
 		.SetTooltipText(LOCTEXT("ColorViewerTooltipText", "Open the Color Viewer tab, a sample application for the ImageWidgets plugin."))
 		.SetGroup(WorkspaceMenu::GetMenuStructure().GetDeveloperToolsMiscCategory())
-		.SetIcon(FSlateIcon(FAppStyle::GetAppStyleSetName(), "ClassIcon.Texture2D"));
+		.SetIcon(FSlateIcon(FAppStyle::GetAppStyleSetName(), "Icons.Color"));
 
 		if (WorkspaceGroup.IsValid())
 		{
diff --git a/Engine/Plugins/Experimental/ImageWidgets/Source/ImageWidgets/Private/SImageCatalog.cpp b/Engine/Plugins/Experimental/ImageWidgets/Source/ImageWidgets/Private/SImageCatalog.cpp
--- a/Engine/Plugins/Experimental/ImageWidgets/Source/ImageWidgets/Private/SImageCatalog.cpp
+++ b/Engine/Plugins/Experimental/ImageWidgets/Source/ImageWidgets/Private/SImageCatalog.cpp
@@ -4,9 +4,11 @@
 
 #include <Algo/Accumulate.h>
 #include <Brushes/SlateColorBrush.h>
+#include <Framework/Application/SlateApplication.h>
 #include <Styling/StyleColors.h>
-#include <Widgets/Layout/SSeparator.h>
+#include <Widgets/Layout/SScrollBox.h>
 #include <Widgets/Views/SListView.h>
+#include <Widgets/Input/SButton.h>
 
 #include "SImageCatalogItem.h"
 
@@ -31,6 +33,7 @@ struct FGroup
 	FText Heading;
 	TSharedPtr<SListView<FItemType>> ListView;
 	TArray<FItemType> Items;
+	bool bIsExpanded = true;
 };
 
 struct FItemLookup
@@ -43,8 +46,8 @@ class FImpl
 {
 public:
 	FImpl(FName DefaultGroupName, const FText& DefaultGroupHeading, ESelectionMode::Type SelectionMode, bool bAllowSelectionAcrossGroups, bool bShowEmptyGroups,
-	      const SImageCatalog::FOnItemSelected& OnItemSelected, const SImageCatalog::FOnGetContextMenu& OnGetContextMenu,
-	      TSharedPtr<SVerticalBox> GroupsLayout);
+	      const SImageCatalog::FOnItemSelected& OnItemSelected, const SImageCatalog::FOnGetGroupContextMenu& OnGetGroupContextMenu,
+	      const SImageCatalog::FOnGetItemsContextMenu& OnGetItemsContextMenu, TSharedPtr<SScrollBox> Layout);
 
 	FName GetDefaultGroupName() const { return DefaultGroupName; }
 	bool AddGroup(FName Name, const FText& Heading, const FName* BeforeGroupWithThisName);
@@ -62,7 +65,7 @@ public:
 	TSharedPtr<const FImageCatalogItemData> GetItemAt(int32 Index, const FName* GroupName) const;
 	TOptional<FGuid> GetItemGuidAt(int32 Index, const FName* GroupName) const;
 	bool UpdateItem(const FGuid& Guid, const FSlateBrush* Thumbnail, const FText* Name, const FText* Info, const FText* ToolTip);
-	
+
 	bool SelectItem(const FGuid& Guid, bool bSelected);
 	bool ClearSelection(const FName* GroupName);
 
@@ -70,7 +73,6 @@ public:
 	TOptional<FName> GetItemGroupName(const FGuid& Guid) const;
 
 private:
-	
 	bool AddGroup(FName Name, const FText& Heading, int32 Index);
 
 	const FGroup* FindGroup(FName Name) const;
@@ -85,7 +87,8 @@ private:
 		return const_cast<FItemLookup*>(const_cast<const FImpl*>(this)->FindLookup(Guid));
 	}
 
-	TSharedPtr<SWidget> OnContextMenuOpening() const;
+	TSharedPtr<SWidget> OnGroupContextMenuOpening(FName Name) const;
+	TSharedPtr<SWidget> OnItemsContextMenuOpening() const;
 	void SortSelection(TArray<FGuid>& Selection) const;
 	void UpdateGroupMapping(int32 StartingIndex);
 	void UpdateItemMapping(TArray<FItemType>& Items, int32 StartingIndex);
@@ -96,31 +99,34 @@ private:
 	const bool bAllowSelectionAcrossGroups;
 	const bool bShowEmptyGroups;
 	SImageCatalog::FOnItemSelected OnItemSelected;
-	SImageCatalog::FOnGetContextMenu OnGetContextMenu;
-	TSharedPtr<SVerticalBox> GroupsLayout;
+	SImageCatalog::FOnGetGroupContextMenu OnGetGroupContextMenu;
+	SImageCatalog::FOnGetItemsContextMenu OnGetItemsContextMenu;
+	TSharedPtr<SScrollBox> Layout;
 	TArray<TUniquePtr<FGroup>> Groups;
 	TMap<FName, int32> GroupMapping;
 	TMap<FGuid, FItemLookup> ItemMapping;
 };
 
 FImpl::FImpl(const FName DefaultGroupName, const FText& DefaultGroupHeading, const ESelectionMode::Type SelectionMode, const bool bAllowSelectionAcrossGroups,
-             const bool bShowEmptyGroups, const SImageCatalog::FOnItemSelected& OnItemSelected, const SImageCatalog::FOnGetContextMenu& OnGetContextMenu,
-             TSharedPtr<SVerticalBox> GroupsLayout)
+             const bool bShowEmptyGroups, const SImageCatalog::FOnItemSelected& OnItemSelected,
+             const SImageCatalog::FOnGetGroupContextMenu& OnGetGroupContextMenu, const SImageCatalog::FOnGetItemsContextMenu& OnGetItemsContextMenu,
+             TSharedPtr<SScrollBox> Layout)
 	: DefaultGroupName(DefaultGroupName)
 	, DefaultGroupHeading(DefaultGroupHeading)
 	, SelectionMode(SelectionMode)
 	, bAllowSelectionAcrossGroups(bAllowSelectionAcrossGroups)
 	, bShowEmptyGroups(bShowEmptyGroups)
 	, OnItemSelected(OnItemSelected)
-	, OnGetContextMenu(OnGetContextMenu)
-	, GroupsLayout(MoveTemp(GroupsLayout))
+	, OnGetGroupContextMenu(OnGetGroupContextMenu)
+	, OnGetItemsContextMenu(OnGetItemsContextMenu)
+	, Layout(MoveTemp(Layout))
 {
 	AddGroup(DefaultGroupName, DefaultGroupHeading, nullptr);
 }
 
 bool FImpl::AddGroup(const FName Name, const FText& Heading, const FName* BeforeGroupWithThisName)
 {
-	if (GroupsLayout.IsValid())
+	if (Layout.IsValid())
 	{
 		if (BeforeGroupWithThisName != nullptr)
 		{
@@ -155,7 +161,8 @@ TOptional<TArray<FGuid>> FImpl::RemoveGroup(const FName Name, const FName* Group
 		if (NewGroupIndex)
 		{
 			check(0 <= *NewGroupIndex && *NewGroupIndex < Groups.Num());
-			TArray<FItemType>& NewGroupItems = Groups[*NewGroupIndex]->Items;
+			FGroup& NewGroup = *Groups[*NewGroupIndex];
+			TArray<FItemType>& NewGroupItems = NewGroup.Items;
 
 			for (FItemType& Item : GroupItems)
 			{
@@ -168,6 +175,8 @@ TOptional<TArray<FGuid>> FImpl::RemoveGroup(const FName Name, const FName* Group
 
 				NewGroupItems.Emplace(MoveTemp(Item));
 			}
+
+			NewGroup.ListView->RequestListRefresh();
 		}
 		else
 		{
@@ -178,15 +187,16 @@ TOptional<TArray<FGuid>> FImpl::RemoveGroup(const FName Name, const FName* Group
 			}
 		}
 
-		check(GroupsLayout->GetChildren() && *GroupIndex < GroupsLayout->NumSlots());
-		GroupsLayout->RemoveSlot(GroupsLayout->GetChildren()->GetChildAt(*GroupIndex));		
+		check(Layout->GetChildren() && *GroupIndex < Layout->NumSlots());
+		const TSharedRef<SWidget>& WidgetToRemove = Layout->GetSlot(*GroupIndex).GetWidget();
+		Layout->RemoveSlot(WidgetToRemove);
 		Groups.RemoveAt(*GroupIndex);
 		GroupMapping.Remove(Name);
 		UpdateGroupMapping(*GroupIndex);
 
 		return {AffectedGuids};
 	}
-	
+
 	return {};
 }
 
@@ -315,7 +325,7 @@ bool FImpl::MoveItem(const FGuid& Guid, const FName* GroupName, const FGuid* Bef
 				if (*GroupToIndex != Lookup->GroupIndex)
 				{
 					FItemType Item = MoveTemp(GroupFrom.Items[Lookup->ItemIndex]);
-				
+
 					GroupFrom.Items.RemoveAt(Lookup->ItemIndex, EAllowShrinking::No);
 					UpdateItemMapping(GroupFrom.Items, Lookup->ItemIndex);
 
@@ -426,10 +436,10 @@ bool FImpl::UpdateItem(const FGuid& Guid, const FSlateBrush* Thumbnail, const FT
 		if (Name) Item.Name = *Name;
 		if (Info) Item.Info = *Info;
 		if (ToolTip) Item.ToolTip = *ToolTip;
-		
+
 		return true;
 	}
-	
+
 	return false;
 }
 
@@ -493,9 +503,19 @@ TOptional<FName> FImpl::GetItemGroupName(const FGuid& Guid) const
 	return {};
 }
 
-TSharedPtr<SWidget> FImpl::OnContextMenuOpening() const
+TSharedPtr<SWidget> FImpl::OnGroupContextMenuOpening(FName Name) const
+{
+	if (!OnGetGroupContextMenu.IsBound())
+	{
+		return SNullWidget::NullWidget;
+	}
+
+	return OnGetGroupContextMenu.Execute(Name);
+}
+
+TSharedPtr<SWidget> FImpl::OnItemsContextMenuOpening() const
 {
-	if (!OnGetContextMenu.IsBound())
+	if (!OnGetItemsContextMenu.IsBound())
 	{
 		return SNullWidget::NullWidget;
 	}
@@ -513,7 +533,7 @@ TSharedPtr<SWidget> FImpl::OnContextMenuOpening() const
 
 	SortSelection(SelectedGuids);
 
-	return OnGetContextMenu.Execute(SelectedGuids);
+	return OnGetItemsContextMenu.Execute(SelectedGuids);
 }
 
 bool FImpl::AddGroup(FName Name, const FText& Heading, const int32 Index)
@@ -530,15 +550,20 @@ bool FImpl::AddGroup(FName Name, const FText& Heading, const int32 Index)
 		{
 			return Group.Heading;
 		};
-		
+
 		const auto GetHeadingVisibility = [bShowEmptyGroups = bShowEmptyGroups, &Group]
 		{
 			return !Group.Heading.IsEmpty() && (bShowEmptyGroups || !Group.Items.IsEmpty()) ? EVisibility::Visible : EVisibility::Collapsed;
 		};
-		
+
+		const auto GetHeadingHeight = [bShowEmptyGroups = bShowEmptyGroups, &Group]
+		{
+			return !Group.Heading.IsEmpty() && (bShowEmptyGroups || !Group.Items.IsEmpty()) ? 26.0f : 0.0f;
+		};
+
 		const auto GetVisibility = [bShowEmptyGroups = bShowEmptyGroups, &Group]
 		{
-			return bShowEmptyGroups || !Group.Items.IsEmpty() ? EVisibility::Visible : EVisibility::Collapsed;
+			return bShowEmptyGroups || (Group.bIsExpanded && !Group.Items.IsEmpty()) ? EVisibility::Visible : EVisibility::Collapsed;
 		};
 
 		const auto GenerateItemRow = [](const FItemType& ItemData, const TSharedRef<STableViewBase>& OwnerTable)
@@ -587,36 +612,104 @@ bool FImpl::AddGroup(FName Name, const FText& Heading, const int32 Index)
 			}
 		};
 
-		const auto OpenContextMenu = [this] { return OnContextMenuOpening(); };
+		const auto OpenItemsContextMenu = [this]
+		{
+			return OnItemsContextMenuOpening();
+		};
+
+		TSharedPtr<SWidget> GroupHeader;
 
-		GroupsLayout->InsertSlot(Index)
-			.AutoHeight()
+		Layout->InsertSlot(Index)
+		[
+			SNew(SBorder)
+			.BorderImage(FAppStyle::Get().GetBrush("DetailsView.GridLine"))
+			.Padding(0.0f, 0.0f, 0.0f, 1.0f)
 			[
 				SNew(SVerticalBox)
-					+ SVerticalBox::Slot()
-						.AutoHeight()
-						.Padding(2.0f, 4.0f, 2.0f, 4.0f)
+				+ SVerticalBox::Slot()
+				.AutoHeight()
+				[
+					SAssignNew(GroupHeader, SBox)
+					.MinDesiredHeight_Lambda(GetHeadingHeight)
+					.Visibility_Lambda(GetHeadingVisibility)
+					[
+						SNew(SBorder)
+						.BorderImage(FAppStyle::Get().GetBrush("DetailsView.CategoryTop"))
+						.Padding(0.0f)
 						[
-							SNew(STextBlock)
+							SNew(SHorizontalBox)
+							+ SHorizontalBox::Slot()
+							.HAlign(HAlign_Left)
+							.VAlign(VAlign_Center)
+							.Padding(2.0f, 0.0f, 0.0f, 0.0f)
+							.AutoWidth()
+							[
+								SNew(SButton)
+								.ButtonStyle(FCoreStyle::Get(), "NoBorder")
+								.VAlign(VAlign_Center)
+								.HAlign(HAlign_Center)
+								.ClickMethod(EButtonClickMethod::MouseDown)
+								.OnClicked_Lambda([&bIsExpanded = Group.bIsExpanded]()
+								{
+									bIsExpanded = !bIsExpanded;
+									return FReply::Handled();
+								})
+								.ContentPadding(0.0f)
+								.IsFocusable(false)
+								[
+									SNew(SImage)
+									.Image_Lambda([&bIsExpanded = Group.bIsExpanded]()
+									{
+										return FAppStyle::Get().GetBrush(bIsExpanded ? "TreeArrow_Expanded" : "TreeArrow_Collapsed");
+									})
+									.ColorAndOpacity(FSlateColor::UseSubduedForeground())
+								]
+							]
+							+ SHorizontalBox::Slot()
+							.VAlign(VAlign_Center)
+							.FillWidth(1.0f)
+							.Padding(2.0f, 2.0f, 2.0f, 2.0f)
+							[
+								SNew(STextBlock)
 								.Text_Lambda(GetHeading)
-								.Visibility_Lambda(GetHeadingVisibility)
-								.TextStyle(FAppStyle::Get(), "DetailsView.CategoryTextStyle")
+								.Font(FAppStyle::Get().GetFontStyle("DetailsView.CategoryFontStyle"))
+							]
 						]
-					+ SVerticalBox::Slot()
-						.AutoHeight()
-						.Padding(0.0f, 0.0f, 0.0f, 4.0f)
-						[
-							SAssignNew(Group.ListView, SListView<FItemType>)
-								.ListItemsSource(&Group.Items)
-								.Visibility_Lambda(GetVisibility)
-								.ScrollbarVisibility(EVisibility::Visible)
-                                .OnGenerateRow_Lambda(GenerateItemRow)
-								.SelectionMode(SelectionMode)
-								.ClearSelectionOnClick(false)
-								.OnSelectionChanged_Lambda(SelectionChanged)
-							    .OnContextMenuOpening_Lambda(OpenContextMenu)
-						]
-			];
+					]
+				]
+				+ SVerticalBox::Slot()
+				[
+					SAssignNew(Group.ListView, SListView<FItemType>)
+					.ListItemsSource(&Group.Items)
+					.Visibility_Lambda(GetVisibility)
+					.ScrollbarVisibility(EVisibility::Collapsed)
+					.OnGenerateRow_Lambda(GenerateItemRow)
+					.SelectionMode(SelectionMode)
+					.ClearSelectionOnClick(false)
+					.OnSelectionChanged_Lambda(SelectionChanged)
+					.OnContextMenuOpening_Lambda(OpenItemsContextMenu)
+				]
+			]
+		];
+
+		GroupHeader->SetOnMouseButtonUp(FPointerEventHandler::CreateLambda([this, Name, GroupHeader](const FGeometry& Geometry, const FPointerEvent& Event)
+		{
+			if (Event.GetEffectingButton() == EKeys::RightMouseButton)
+			{
+				TSharedPtr<SWidget> MenuContent = OnGroupContextMenuOpening(Name);
+
+				if (MenuContent.IsValid())
+				{
+					const FWidgetPath WidgetPath = Event.GetEventPath() != nullptr ? *Event.GetEventPath() : FWidgetPath();
+					FSlateApplication::Get().PushMenu(GroupHeader->AsShared(), WidgetPath, MenuContent.ToSharedRef(), Event.GetScreenSpacePosition(),
+					                                  FPopupTransitionEffect(FPopupTransitionEffect::ContextMenu));
+				}
+
+				return FReply::Handled().ReleaseMouseCapture();
+			}
+
+			return FReply::Unhandled();
+		}));
 
 		return true;
 	}
@@ -657,7 +750,7 @@ void FImpl::SortSelection(TArray<FGuid>& Selection) const
 		GroupOffsets.Add(Offset);
 		Offset += Group->Items.Num();
 	}
-	
+
 	Algo::Sort(Selection, [&ItemMapping = ItemMapping, &GroupOffsets](const FGuid& A, const FGuid& B)
 	{
 		const FItemLookup *const LookupA = ItemMapping.Find(A);
@@ -669,7 +762,7 @@ void FImpl::SortSelection(TArray<FGuid>& Selection) const
 
 		const int32 PositionA = GroupOffsets[LookupA->GroupIndex] + LookupA->ItemIndex;
 		const int32 PositionB = GroupOffsets[LookupB->GroupIndex] + LookupB->ItemIndex;
-		
+
 		return PositionA < PositionB;
 	});
 }
@@ -704,15 +797,15 @@ void FImpl::UpdateItemMapping(TArray<FItemType>& Items, const int32 StartingInde
 
 void SImageCatalog::Construct(const FArguments& Args)
 {
-	TSharedPtr<SVerticalBox> GroupsLayout;
+	TSharedPtr<SScrollBox> Layout;
 
 	ChildSlot
 	[
-		SAssignNew(GroupsLayout, SVerticalBox)
+		SAssignNew(Layout, SScrollBox)
 	];
 
 	Impl = MakePimpl<FImpl>(Args._DefaultGroupName, Args._DefaultGroupHeading, Args._SelectionMode, Args._bAllowSelectionAcrossGroups,
-	                        Args._bShowEmptyGroups, Args._OnItemSelected, Args._OnGetContextMenu, GroupsLayout);
+	                        Args._bShowEmptyGroups, Args._OnItemSelected, Args._OnGetGroupContextMenu, Args._OnGetItemsContextMenu, Layout);
 }
 
 FName SImageCatalog::GetDefaultGroupName() const
diff --git a/Engine/Plugins/Experimental/ImageWidgets/Source/ImageWidgets/Public/SImageCatalog.h b/Engine/Plugins/Experimental/ImageWidgets/Source/ImageWidgets/Public/SImageCatalog.h
--- a/Engine/Plugins/Experimental/ImageWidgets/Source/ImageWidgets/Public/SImageCatalog.h
+++ b/Engine/Plugins/Experimental/ImageWidgets/Source/ImageWidgets/Public/SImageCatalog.h
@@ -45,11 +45,17 @@ namespace UE::ImageWidgets
 		 */
 		DECLARE_DELEGATE_OneParam(FOnItemSelected, const FGuid&)
 
+		/**
+		 * Delegate that gets called for creating a context menu for a group.
+		 * Return @see SNullWidget::NullWidget to not show a context menu.
+		 */
+		DECLARE_DELEGATE_RetVal_OneParam(TSharedPtr<SWidget>, FOnGetGroupContextMenu, FName)
+
 		/**
 		 * Delegate that gets called for creating a context menu for a set of selected items.
-		 * Return @see SWidget::NullWidget to not show a context menu.
+		 * Return @see SNullWidget::NullWidget to not show a context menu.
 		 */
-		DECLARE_DELEGATE_RetVal_OneParam(TSharedPtr<SWidget>, FOnGetContextMenu, const TArray<FGuid>&)
+		DECLARE_DELEGATE_RetVal_OneParam(TSharedPtr<SWidget>, FOnGetItemsContextMenu, const TArray<FGuid>&)
 
 		SLATE_BEGIN_ARGS(SImageCatalog)
 				: _DefaultGroupName(NAME_None)
@@ -62,15 +68,18 @@ namespace UE::ImageWidgets
 
 			/** Name of the default group, i.e. the group that gets used when no group is explicitly specified. */
 			SLATE_ARGUMENT(FName, DefaultGroupName)
-			
+
 			/** Header text for the default group. */
 			SLATE_ARGUMENT(FText, DefaultGroupHeading)
 
 			/** Delegate that gets called when an item is selected in the catalog. */
 			SLATE_EVENT(FOnItemSelected, OnItemSelected)
 
+			/** Delegate that gets called for creating a context menu for a group. */
+			SLATE_EVENT(FOnGetGroupContextMenu, OnGetGroupContextMenu)
+
 			/** Delegate that gets called for creating a context menu for a set of selected items. */
-			SLATE_EVENT(FOnGetContextMenu, OnGetContextMenu)
+			SLATE_EVENT(FOnGetItemsContextMenu, OnGetItemsContextMenu)
 
 			/** Defines the selection behavior within an item list, e.g. only allow single item selection or do not allow any selection. */
 			SLATE_ARGUMENT(ESelectionMode::Type, SelectionMode)
@@ -92,7 +101,7 @@ namespace UE::ImageWidgets
 		 * @returns The name of the default group.
 		 */
 		IMAGEWIDGETS_API FName GetDefaultGroupName() const;
-		
+
 		/**
 		 * Add a custom group to the catalog.
 		 * @param Name Unique identifier for the group.
@@ -275,7 +284,7 @@ namespace UE::ImageWidgets
 		 * @returns The pointer to the item or an invalid pointer if the group does not exist or the index is invalid.
 		 */
 		IMAGEWIDGETS_API TSharedPtr<const FImageCatalogItemData> GetItemAt(int32 Index, FName Group) const;
-		
+
 		/**
 		 * Retrieve the unique identifier of an existing item at the given index within the default group.
 		 * @param Index Index of the item in the default group; should be at least 0 and less than @see NumItems(@see GetDefaultGroup()).
@@ -290,7 +299,7 @@ namespace UE::ImageWidgets
 		 * @return Unique identifier of the item or no value if the group does not exist or the index is invalid.
 		 */
 		IMAGEWIDGETS_API TOptional<FGuid> GetItemGuidAt(int32 Index, FName Group) const;
-		
+
 		/**
 		 * Select an existing item.
 		 * @param Guid The unique identifier of the existing item.
