ColorViewerWidget: Fixed duplicate LOCTEXT key 'DeleteGroupItems'

diff --git a/Engine/Plugins/Experimental/ImageWidgets/Source/ImageWidgets/Private/ColorViewerSample/ColorViewerWidget.cpp b/Engine/Plugins/Experimental/ImageWidgets/Source/ImageWidgets/Private/ColorViewerSample/ColorViewerWidget.cpp
--- a/Engine/Plugins/Experimental/ImageWidgets/Source/ImageWidgets/Private/ColorViewerSample/ColorViewerWidget.cpp
+++ b/Engine/Plugins/Experimental/ImageWidgets/Source/ImageWidgets/Private/ColorViewerSample/ColorViewerWidget.cpp
@@ -181,7 +181,7 @@ namespace UE::ImageWidgets::Sample
 			MenuBuilder.AddMenuEntry(
 				NumItems == 1
 					? FText::Format(LOCTEXT("DeleteGroupItem", "Delete single item in group"), NumItems)
-					: FText::Format(LOCTEXT("DeleteGroupItems", "Delete all {0} items in group"), NumItems),
+					: FText::Format(LOCTEXT("DeleteAllGroupItems", "Delete all {0} items in group"), NumItems),
 				LOCTEXT("DeleteGroupItems", "Deletes all items in this group."),
 				FSlateIcon(FAppStyle::GetAppStyleSetName(), "GenericCommands.Delete"),
 				FUIAction(FExecuteAction::CreateLambda(
