ImageWidgets: Fix viewport showing incorrect zoom label when using DPI scaling

diff --git a/Engine/Plugins/Experimental/ImageWidgets/Source/ImageWidgets/Private/SImageViewport.cpp b/Engine/Plugins/Experimental/ImageWidgets/Source/ImageWidgets/Private/SImageViewport.cpp
--- a/Engine/Plugins/Experimental/ImageWidgets/Source/ImageWidgets/Private/SImageViewport.cpp
+++ b/Engine/Plugins/Experimental/ImageWidgets/Source/ImageWidgets/Private/SImageViewport.cpp
@@ -176,21 +176,12 @@ namespace UE::ImageWidgets
 			return DrawSettings.Get({});
 		};
 
-		auto GetDPIScaleFactor = [this]
-		{
-			if (const TSharedPtr<SWindow> TopLevelWindow = FSlateApplication::Get().FindWidgetWindow(SharedThis(this)); TopLevelWindow.IsValid())
-			{
-				return TopLevelWindow->GetDPIScaleFactor();
-			}
-			return 1.0f;
-		};
-
 		ImageViewportClient = MakeShareable(new FImageViewportClient(
 			StaticCastWeakPtr<SEditorViewport>(AsWeak()),
 			FGetImageSize::CreateLambda(GetImageSize),
 			FDrawImage::CreateLambda(DrawImage),
 			FGetDrawSettings::CreateLambda(GetDrawSettings),
-			FGetDPIScaleFactor::CreateLambda(GetDPIScaleFactor),
+			FGetDPIScaleFactor::CreateRaw(this, &SImageViewport::GetDPIScaleFactor),
 			ABComparison.Get(),
 			InArgs._ControllerSettings)
 			);
@@ -294,16 +285,15 @@ namespace UE::ImageWidgets
 		};
 
 		SAssignNew(ImageViewportToolbar, SImageViewportToolbar, StaticCastSharedPtr<FImageViewportClient>(Client), CommandList,
-						  SImageViewportToolbar::FConstructParameters
-						  {
-							  SImageViewportToolbar::FHasImage::CreateLambda(HasImage),
-							  SImageViewportToolbar::FNumMips::CreateLambda(GetNumMips),
-							  SImageViewportToolbar::FImageGuid::CreateLambda(GetImageGuid),
-							  SImageViewportToolbar::FGetOverlaySettings::CreateLambda(GetOverlaySettings),
-							  bABComparisonEnabled ? ABComparison.Get() : nullptr,
-							  ToolbarExtender
-						  });
-						  
+		           SImageViewportToolbar::FConstructParameters{
+			           SImageViewportToolbar::FHasImage::CreateLambda(HasImage),
+			           SImageViewportToolbar::FNumMips::CreateLambda(GetNumMips),
+			           SImageViewportToolbar::FImageGuid::CreateLambda(GetImageGuid),
+			           FGetDPIScaleFactor::CreateRaw(this, &SImageViewport::GetDPIScaleFactor),
+			           SImageViewportToolbar::FGetOverlaySettings::CreateLambda(GetOverlaySettings),
+			           bABComparisonEnabled ? ABComparison.Get() : nullptr,
+			           ToolbarExtender
+		           });
 
 		ToolbarExtender.Reset();
 
@@ -319,6 +309,15 @@ namespace UE::ImageWidgets
 		];
 	}
 
+	float SImageViewport::GetDPIScaleFactor() const
+	{
+		if (const TSharedPtr<SWindow> TopLevelWindow = FSlateApplication::Get().FindWidgetWindow(SharedThis(this)); TopLevelWindow.IsValid())
+		{
+			return TopLevelWindow->GetDPIScaleFactor();
+		}
+		return 1.0f;
+	}
+
 	FText SImageViewport::GetPickerLabel() const
 	{
 		const IImageViewer::FImageInfo ImageInfo = ImageViewer->GetCurrentImageInfo();
@@ -468,7 +467,7 @@ namespace UE::ImageWidgets
 		ApplyHook(FName("StatusBarRight"), EExtensionHook::After);
 
 		StatusBarExtender.Reset();
-		
+
 		return HorizontalBox.ToSharedRef();
 	}
 }
diff --git a/Engine/Plugins/Experimental/ImageWidgets/Source/ImageWidgets/Private/SImageViewportToolbar.cpp b/Engine/Plugins/Experimental/ImageWidgets/Source/ImageWidgets/Private/SImageViewportToolbar.cpp
--- a/Engine/Plugins/Experimental/ImageWidgets/Source/ImageWidgets/Private/SImageViewportToolbar.cpp
+++ b/Engine/Plugins/Experimental/ImageWidgets/Source/ImageWidgets/Private/SImageViewportToolbar.cpp
@@ -32,10 +32,12 @@ namespace UE::ImageWidgets
 		HasImage = MoveTemp(Parameters.HasImage);
 		NumMips = MoveTemp(Parameters.NumMips);
 		ImageGuid = MoveTemp(Parameters.ImageGuid);
+		GetDPIScaleFactor = MoveTemp(Parameters.GetDPIScaleFactor);
 		GetOverlaySettings = MoveTemp(Parameters.GetOverlaySettings);
 		check(HasImage.IsBound());
 		check(NumMips.IsBound());
 		check(ImageGuid.IsBound());
+		check(GetDPIScaleFactor.IsBound());
 		check(GetOverlaySettings.IsBound());
 
 		ABComparison = Parameters.ABComparison;
@@ -198,7 +200,9 @@ namespace UE::ImageWidgets
 	{
 		const FImageViewportController::FZoomSettings ZoomSettings = ViewportClient->GetZoom();
 
-		const double Zoom = ZoomSettings.Zoom;
+		const float DPIScaleFactor = GetDPIScaleFactor.Execute();
+
+		const double Zoom = ZoomSettings.Zoom / DPIScaleFactor;
 		FNumberFormattingOptions FormattingOptions;
 		FormattingOptions.SetMaximumFractionalDigits(Zoom < 1.0 ? (Zoom < 0.1 ? 2 : 1) : 0);
 		const FText ZoomPercentage = FText::AsPercent(Zoom, &FormattingOptions);
diff --git a/Engine/Plugins/Experimental/ImageWidgets/Source/ImageWidgets/Private/SImageViewportToolbar.h b/Engine/Plugins/Experimental/ImageWidgets/Source/ImageWidgets/Private/SImageViewportToolbar.h
--- a/Engine/Plugins/Experimental/ImageWidgets/Source/ImageWidgets/Private/SImageViewportToolbar.h
+++ b/Engine/Plugins/Experimental/ImageWidgets/Source/ImageWidgets/Private/SImageViewportToolbar.h
@@ -3,6 +3,7 @@
 #pragma once
 
 #include "CoreMinimal.h"
+#include "ImageViewportClient.h"
 #include "SImageViewport.h"
 #include "SViewportToolBar.h"
 
@@ -14,7 +15,6 @@ class SViewportToolBarComboMenu;
 namespace UE::ImageWidgets
 {
 	class FImageABComparison;
-	class FImageViewportClient;
 
 	/**
 	 * Extendable toolbar for the image viewport.
@@ -38,6 +38,7 @@ namespace UE::ImageWidgets
 			FHasImage HasImage;
 			FNumMips NumMips;
 			FImageGuid ImageGuid;
+			FGetDPIScaleFactor GetDPIScaleFactor;
 			FGetOverlaySettings GetOverlaySettings;
 			FImageABComparison* ABComparison = nullptr;
 			TSharedPtr<FExtender> ToolbarExtender;
@@ -66,6 +67,7 @@ namespace UE::ImageWidgets
 		FHasImage HasImage;
 		FNumMips NumMips;
 		FImageGuid ImageGuid;
+		FGetDPIScaleFactor GetDPIScaleFactor;
 		FGetOverlaySettings GetOverlaySettings;
 		FImageABComparison* ABComparison = nullptr;
 	};
diff --git a/Engine/Plugins/Experimental/ImageWidgets/Source/ImageWidgets/Public/SImageViewport.h b/Engine/Plugins/Experimental/ImageWidgets/Source/ImageWidgets/Public/SImageViewport.h
--- a/Engine/Plugins/Experimental/ImageWidgets/Source/ImageWidgets/Public/SImageViewport.h
+++ b/Engine/Plugins/Experimental/ImageWidgets/Source/ImageWidgets/Public/SImageViewport.h
@@ -242,6 +242,9 @@ namespace UE::ImageWidgets
 		/** Create the status bar widgets. */
 		TSharedRef<SWidget> MakeStatusBar(TSharedRef<SOverlay> Overlay);
 
+		/** Provides the UI scale factor. */
+		float GetDPIScaleFactor() const;
+
 		/** Provides the text for the color picker display in the status bar. */
 		FText GetPickerLabel() const;
 
