[ImageWidget]
* Fixed compile error when compiling with /Ob3

diff --git a/Engine/Plugins/Experimental/ImageWidgets/Source/ImageWidgets/Private/ImageViewportController.cpp b/Engine/Plugins/Experimental/ImageWidgets/Source/ImageWidgets/Private/ImageViewportController.cpp
--- a/Engine/Plugins/Experimental/ImageWidgets/Source/ImageWidgets/Private/ImageViewportController.cpp
+++ b/Engine/Plugins/Experimental/ImageWidgets/Source/ImageWidgets/Private/ImageViewportController.cpp
@@ -83,7 +83,7 @@ namespace UE::ImageWidgets
 	FImageViewportController::FImageViewportController(const EZoomMode DefaultZoomMode)
 		: DefaultZoomMode(DefaultZoomMode)
 	{
-		Reset({0, 0}, {0, 0});
+		Reset(FIntPoint::ZeroValue, FVector2d::Zero());
 	}
 
 	void FImageViewportController::Pan(const FVector2d ViewportDeltaWithDPIScaling)
