ImageWidgets: fix local variable name in example code

#rnx
#Synced-CL 49059777

diff --git a/Engine/Plugins/Experimental/ImageWidgets/Source/ImageWidgets/Private/ImageWidgetsModule.cpp b/Engine/Plugins/Experimental/ImageWidgets/Source/ImageWidgets/Private/ImageWidgetsModule.cpp
--- a/Engine/Plugins/Experimental/ImageWidgets/Source/ImageWidgets/Private/ImageWidgetsModule.cpp
+++ b/Engine/Plugins/Experimental/ImageWidgets/Source/ImageWidgets/Private/ImageWidgetsModule.cpp
@@ -89,10 +89,10 @@ namespace UE::ImageWidgets
 
 	TSharedRef<SDockTab> FImageWidgetsModule::MakeColorViewerTab(const FSpawnTabArgs&)
 	{
-		TSharedRef<SDockTab> WidgetReflectorTab = SNew(SDockTab)
+		TSharedRef<SDockTab> Tab = SNew(SDockTab)
 			.TabRole(NomadTab);
-		WidgetReflectorTab->SetContent(GetColorViewer());
-		return WidgetReflectorTab;
+		Tab->SetContent(GetColorViewer());
+		return Tab;
 	}
 
 	TSharedRef<SWidget> FImageWidgetsModule::GetColorViewer()
