ImageWidgets: Trivial fix in sample code

#Synced-CL 49891327

diff --git a/Engine/Plugins/Experimental/ImageWidgets/Source/ImageWidgets/Private/ColorViewerSample/ColorViewerWidget.cpp b/Engine/Plugins/Experimental/ImageWidgets/Source/ImageWidgets/Private/ColorViewerSample/ColorViewerWidget.cpp
--- a/Engine/Plugins/Experimental/ImageWidgets/Source/ImageWidgets/Private/ColorViewerSample/ColorViewerWidget.cpp
+++ b/Engine/Plugins/Experimental/ImageWidgets/Source/ImageWidgets/Private/ColorViewerSample/ColorViewerWidget.cpp
@@ -270,21 +270,11 @@ namespace UE::ImageWidgets::Sample
 						})));
 		}
 
-#if 0
-		MenuBuilder.AddMenuEntry(
-			LOCTEXT("RenameItem", "Rename"),
-			LOCTEXT("RenameItemTooltip", "Renames a single selected item."),
-			FSlateIcon(FAppStyle::GetAppStyleSetName(), "GenericCommands.Rename"),
-			FUIAction(
-				FExecuteAction::CreateLambda([]{}),
-				FCanExecuteAction::CreateLambda([bSingleItem = Guids.Num() == 1] { return bSingleItem; })));
-#endif
-
 		MenuBuilder.AddSeparator();
 
 		MenuBuilder.AddMenuEntry(
 			LOCTEXT("DeleteColors", "Delete"),
-			LOCTEXT("DeleteColorsTooltip", "Deletes the selected colors(s) from the catalog."),
+			LOCTEXT("DeleteColorsTooltip", "Deletes the selected color(s) from the catalog."),
 			FSlateIcon(FAppStyle::GetAppStyleSetName(), "GenericCommands.Delete"),
 			FUIAction(FExecuteAction::CreateLambda(
 				[Guids, Catalog = Catalog, ColorViewer = ColorViewer]
