ImageWidgets: Rework viewport scaling

#rn ImageWidgets: Fix issues with layout scaling in SImageViewport

#Synced-CL 51563907

diff --git a/Engine/Plugins/Experimental/ImageWidgets/Source/ImageWidgets/Private/ImageViewportClient.cpp b/Engine/Plugins/Experimental/ImageWidgets/Source/ImageWidgets/Private/ImageViewportClient.cpp
--- a/Engine/Plugins/Experimental/ImageWidgets/Source/ImageWidgets/Private/ImageViewportClient.cpp
+++ b/Engine/Plugins/Experimental/ImageWidgets/Source/ImageWidgets/Private/ImageViewportClient.cpp
@@ -15,7 +15,7 @@
 
 namespace UE::ImageWidgets
 {
-	
+
 UTexture2D* CreateCheckerTexture(const FColor& Color1, const FColor& Color2, const int32 CheckerSize)
 {
 	constexpr EPixelFormat Format = PF_B8G8R8A8;
@@ -26,11 +26,11 @@ UTexture2D* CreateCheckerTexture(const FColor& Color1, const FColor& Color2, con
 
 	UE_LOG(LogImageWidgets, Verbose, TEXT("Creating background checker texture '%s' with Color1=%s, Color2=%s, and Size=%d."),
 	       *CheckerTexture->GetName(), *Color1.ToHex(), *Color2.ToHex(), CheckerSize);
-	
+
 	// TODO Revisit setting the LODGroup, and figure out if there is a better way to do this.
 	// Setting the LOD group seems to be the only way to avoid interpolation when zoomed in. 
 	CheckerTexture->LODGroup = TEXTUREGROUP_Pixels2D;
-	
+
 	TIndirectArray<FTexture2DMipMap>& Mips = CheckerTexture->GetPlatformData()->Mips;
 	int32 MipSize = TextureSize;
 	while (MipSize >= GPixelFormats[Format].BlockSizeX)
@@ -46,7 +46,7 @@ UTexture2D* CreateCheckerTexture(const FColor& Color1, const FColor& Color2, con
 			{
 				const int64 NumBlocks = MipSize / GPixelFormats[Format].BlockSizeX;
 				const int64 NumBytes = NumBlocks * NumBlocks * GPixelFormats[Format].BlockBytes;
-				
+
 				ColorData = reinterpret_cast<FColor*>(MipBulkData.Realloc(NumBytes));
 			}
 			return ColorData;
@@ -99,13 +99,15 @@ void DestroyCheckerTexture(TStrongObjectPtr<UTexture2D>& CheckerTexture)
 }
 
 FImageViewportClient::FImageViewportClient(const TWeakPtr<SEditorViewport>& InEditorViewport, FGetImageSize&& InGetImageSize, FDrawImage&& InDrawImage,
-                                           FGetDrawSettings&& InGetDrawSettings, FGetDPIScaleFactor&& InGetDPIScaleFactor,
-                                           const FImageABComparison* InABComparison, const SImageViewport::FControllerSettings& InControllerSettings)
+                                           FGetDrawSettings&& InGetDrawSettings, FGetLayoutAndDPIScaleFactor&& InGetLayoutAndDPIScaleFactor,
+                                           FGetLayoutScaleFactor&& InGetLayoutScaleFactor, const FImageABComparison* InABComparison,
+                                           const SImageViewport::FControllerSettings& InControllerSettings)
 	: FEditorViewportClient(nullptr, nullptr, InEditorViewport)
 	, GetImageSize(MoveTemp(InGetImageSize))
 	, DrawImage(MoveTemp(InDrawImage))
 	, GetDrawSettings(MoveTemp(InGetDrawSettings))
-	, GetDPIScaleFactor(MoveTemp(InGetDPIScaleFactor))
+	, GetLayoutAndDPIScaleFactor(MoveTemp(InGetLayoutAndDPIScaleFactor))
+	, GetLayoutScaleFactor(MoveTemp(InGetLayoutScaleFactor))
 	, OnInputKey(InControllerSettings.OnInputKey)
 	, ABComparison(InABComparison)
 	, Controller(static_cast<FImageViewportController::EZoomMode>(InControllerSettings.DefaultZoomMode))
@@ -113,6 +115,8 @@ FImageViewportClient::FImageViewportClient(const TWeakPtr<SEditorViewport>& InEd
 	check(GetImageSize.IsBound());
 	check(DrawImage.IsBound());
 	check(GetDrawSettings.IsBound());
+	check(GetLayoutAndDPIScaleFactor.IsBound());
+	check(GetLayoutScaleFactor.IsBound());
 
 	SetRealtime(true);
 }
@@ -138,16 +142,24 @@ void FImageViewportClient::Draw(FViewport* InViewport, FCanvas* Canvas)
 		return;
 	}
 
-	const FVector2d ViewportSizeWithDPIScaling = FVector2d(InViewport->GetSizeXY()) / GetDPIScaleFactor.Execute();	
-	CachedPlacement = GetPlacementProperties(ImageSize, ViewportSizeWithDPIScaling);
+	const float LayoutAndDPIScale = GetLayoutAndDPIScaleFactor.Execute();
+	const float LayoutScale = GetLayoutScaleFactor.Execute();
+
+	// Controller computes placement in unzoomed logical space
+	const FVector2d ScaledViewportSize = FVector2d(InViewport->GetSizeXY()) / LayoutAndDPIScale;
+	CachedPlacement = GetPlacementProperties(ImageSize, ScaledViewportSize);
 	bCachedPlacementIsValid = true;
 
 	const IImageViewer::FDrawProperties::FMip MipProperties = GetMipProperties();
 
+	// Scale placement from unzoomed logical space to RT pixel space for drawing
+	const FVector2d DrawOffset = CachedPlacement.Offset * LayoutScale;
+	const FVector2d DrawSize = CachedPlacement.Size * LayoutScale;
+
 	// Draw border
 	if (DrawSettings.bBorderEnabled)
 	{
-		FCanvasBoxItem Border(CachedPlacement.Offset - DrawSettings.BorderThickness / 2.0, CachedPlacement.Size + DrawSettings.BorderThickness);
+		FCanvasBoxItem Border(DrawOffset - DrawSettings.BorderThickness / 2.0, DrawSize + DrawSettings.BorderThickness);
 		Border.LineThickness = DrawSettings.BorderThickness;
 		Border.SetColor(DrawSettings.BorderColor);
 		Canvas->DrawItem(Border);
@@ -156,7 +168,7 @@ void FImageViewportClient::Draw(FViewport* InViewport, FCanvas* Canvas)
 	// Draw background
 	if (DrawSettings.bBackgroundColorEnabled || DrawSettings.bBackgroundCheckerEnabled)
 	{
-		FCanvasTileItem Background(CachedPlacement.Offset, CachedPlacement.Size, DrawSettings.BackgroundColor);
+		FCanvasTileItem Background(DrawOffset, DrawSize, DrawSettings.BackgroundColor);
 
 		TRefCountPtr<FBatchedElementTexture2DPreviewParameters> BatchedElementParameters;
 
@@ -178,14 +190,15 @@ void FImageViewportClient::Draw(FViewport* InViewport, FCanvas* Canvas)
 		Canvas->DrawItem(Background);
 	}
 
-	const double ABComparisonDividerX = [this]
+	const double ABComparisonDividerX = [this, LayoutScale]
 	{
 		if (bDraggingABComparisonDivider)
 		{
 			FIntPoint MousePos;
 			Viewport->GetMousePos(MousePos);
 
-			FVector2d DividerPos = FVector2d(MousePos) - CachedPlacement.Offset;
+			// Mouse pos is in RT pixel space; convert to unzoomed placement space
+			FVector2d DividerPos = FVector2d(MousePos) / LayoutScale - CachedPlacement.Offset;
 			DividerPos /= CachedPlacement.ZoomFactor;
 
 			const FIntPoint ImageSize = GetImageSize.Execute();
@@ -199,13 +212,14 @@ void FImageViewportClient::Draw(FViewport* InViewport, FCanvas* Canvas)
 		return {ABComparison->GetGuidA(), ABComparison->GetGuidB(), FMath::Clamp(ABComparisonDividerX, 0.0, 1.0)};
 	}();
 
-	// Draw image
-	DrawImage.Execute(InViewport, Canvas, {CachedPlacement, MipProperties, ABComparisonProperties});
+	// Draw image — pass scaled placement for RT pixel space drawing
+	const IImageViewer::FDrawProperties::FPlacement DrawPlacement{DrawOffset, DrawSize, CachedPlacement.ZoomFactor * LayoutScale};
+	DrawImage.Execute(InViewport, Canvas, {DrawPlacement, MipProperties, ABComparisonProperties});
 
 	if (ABComparison->IsActive())
 	{
-		const FVector2D LineStart(CachedPlacement.Offset.X + CachedPlacement.Size.X * ABComparisonDividerX, 0);
-		const FVector2D LineEnd(CachedPlacement.Offset.X + CachedPlacement.Size.X * ABComparisonDividerX, InViewport->GetSizeXY().Y);
+		const FVector2D LineStart(DrawOffset.X + DrawSize.X * ABComparisonDividerX, 0);
+		const FVector2D LineEnd(DrawOffset.X + DrawSize.X * ABComparisonDividerX, InViewport->GetSizeXY().Y);
 		FCanvasLineItem Line(LineStart, LineEnd);
 		Line.LineThickness = 2.0f;
 		Line.SetColor({0.5f, 0.5f, 0.5f, 1.0f});
@@ -263,7 +277,7 @@ bool FImageViewportClient::InputKey(const FInputKeyEventArgs& EventArgs)
 
 		if (EventArgs.Key == EKeys::F)
 		{
-			Controller.Reset(GetImageSize.Execute(), GetViewportSizeWithDPIScaling());
+			Controller.Reset(GetImageSize.Execute(), GetScaledViewportSize());
 
 			return true;
 		}
@@ -313,11 +327,14 @@ void FImageViewportClient::TrackingStopped()
 		FIntPoint DraggingEnd;
 		Viewport->GetMousePos(DraggingEnd);
 
-		FVector2d DraggingEndPos = FVector2d(DraggingEnd) - CachedPlacement.Offset;
+		const float LayoutScale = GetLayoutScaleFactor.Execute();
+
+		// Mouse pos is in RT pixel space; convert to unzoomed placement space
+		FVector2d DraggingEndPos = FVector2d(DraggingEnd) / LayoutScale - CachedPlacement.Offset;
 		DraggingEndPos /= CachedPlacement.ZoomFactor;
 
 		const FIntPoint ImageSize = GetImageSize.Execute();
-		
+
 		ABComparisonDivider = DraggingEndPos.X / ImageSize.X;
 	}
 
@@ -327,9 +344,9 @@ void FImageViewportClient::TrackingStopped()
 
 		FIntPoint DraggingEnd;
 		Viewport->GetMousePos(DraggingEnd);
-		const float DPIScaleFactor = GetDPIScaleFactor.Execute();
 
-		Controller.Pan(FVector2d(DraggingEnd - DraggingStart) / DPIScaleFactor);
+		// Drag delta is in RT pixel space; divide by layout and DPI scale to get unzoomed logical space
+		Controller.Pan(FVector2d(DraggingEnd - DraggingStart) / GetLayoutAndDPIScaleFactor.Execute());
 	}
 
 	RequiredCursorVisibiltyAndAppearance.bDontResetCursor = true;
@@ -349,25 +366,25 @@ void FImageViewportClient::SetMipLevel(int32 InMipLevel)
 
 void FImageViewportClient::ResetController(FIntPoint ImageSize)
 {
-	Controller.Reset(ImageSize, GetViewportSizeWithDPIScaling());
+	Controller.Reset(ImageSize, GetScaledViewportSize());
 }
 
 void FImageViewportClient::ResetZoom(FIntPoint ImageSize)
-{		
+{
 	const FImageViewportController::FZoomSettings ZoomSettings = Controller.GetZoom();
-	Controller.SetZoom(ZoomSettings.Mode, ZoomSettings.Zoom, ImageSize, GetViewportSizeWithDPIScaling());
+	Controller.SetZoom(ZoomSettings.Mode, ZoomSettings.Zoom, ImageSize, GetScaledViewportSize());
 }
 
 FImageViewportController::FZoomSettings FImageViewportClient::GetZoom() const
 {
 	FImageViewportController::FZoomSettings ZoomSettings = Controller.GetZoom();
-	ZoomSettings.Zoom *= GetDPIScaleFactor.Execute();
+	ZoomSettings.Zoom *= GetLayoutAndDPIScaleFactor.Execute();
 	return ZoomSettings;
 }
 
 void FImageViewportClient::SetZoom(FImageViewportController::EZoomMode Mode, double Zoom)
 {
-	Controller.SetZoom(Mode, Zoom, GetImageSize.Execute(), GetViewportSizeWithDPIScaling());
+	Controller.SetZoom(Mode, Zoom, GetImageSize.Execute(), GetScaledViewportSize());
 }
 
 TPair<bool, FVector2d> FImageViewportClient::GetPixelCoordinatesUnderCursor() const
@@ -377,9 +394,11 @@ TPair<bool, FVector2d> FImageViewportClient::GetPixelCoordinatesUnderCursor() co
 		return {false, FVector2d::Zero()};
 	}
 
-	const float DPIScaleFactor = GetDPIScaleFactor.Execute();
-	const FVector2d MousePos((CurrentMousePos.X + 0.5) / DPIScaleFactor, (CurrentMousePos.Y + 0.5) / DPIScaleFactor);
-	const FVector2d CurrentDrag = GetCurrentDragWithDPIScaling();
+	const float LayoutAndDPIScale = GetLayoutAndDPIScaleFactor.Execute();
+
+	// Mouse pos is in RT pixel space; convert to unzoomed logical space (where CachedPlacement lives)
+	const FVector2d MousePos((CurrentMousePos.X + 0.5) / LayoutAndDPIScale, (CurrentMousePos.Y + 0.5) / LayoutAndDPIScale);
+	const FVector2d CurrentDrag = GetCurrentDrag();
 
 	FVector2d CursorPos = FVector2d(MousePos) - CachedPlacement.Offset + CurrentDrag;
 	CursorPos /= CachedPlacement.ZoomFactor;
@@ -387,30 +406,29 @@ TPair<bool, FVector2d> FImageViewportClient::GetPixelCoordinatesUnderCursor() co
 	return {true, CursorPos};
 }
 
-FVector2d FImageViewportClient::GetCurrentDragWithDPIScaling() const
+FVector2d FImageViewportClient::GetCurrentDrag() const
 {
 	if (bDragging)
 	{
 		FIntPoint DraggingEnd;
 		Viewport->GetMousePos(DraggingEnd);
 
-		const float DPIScaleFactor = GetDPIScaleFactor.Execute();
-
-		return FVector2d(DraggingEnd - DraggingStart) / DPIScaleFactor;
+		// Drag delta is in RT pixel space; convert to unzoomed logical space
+		return FVector2d(DraggingEnd - DraggingStart) / GetLayoutAndDPIScaleFactor.Execute();
 	}
-	return FVector2d::Zero();;
+	return FVector2d::Zero();
 }
 
 IImageViewer::FDrawProperties::FPlacement FImageViewportClient::GetPlacementProperties(const FIntPoint ImageSize,
-                                                                                       const FVector2d ViewportSizeWithDPIScaling) const
+                                                                                       const FVector2d ScaledViewportSize) const
 {
-	const FVector2d CurrentDrag = GetCurrentDragWithDPIScaling();
+	const FVector2d CurrentDrag = GetCurrentDrag();
 	const FVector2d Pan = Controller.GetPan(CurrentDrag);
 
 	const double ZoomFactor = Controller.GetZoom().Zoom;
 
 	const FVector2d TileSize = FVector2d(ImageSize) * ZoomFactor;
-	const FVector2d TileOffset = (ViewportSizeWithDPIScaling - TileSize) / 2 + Pan;
+	const FVector2d TileOffset = (ScaledViewportSize - TileSize) / 2 + Pan;
 
 	return {TileOffset, TileSize, ZoomFactor};
 }
@@ -429,14 +447,14 @@ void FImageViewportClient::CreateOrDestroyCheckerTextureIfSettingsChanged(const
 	const FCheckerTextureSettings NewCheckerTextureSettings{
 		DrawSettings.bBackgroundCheckerEnabled, DrawSettings.BackgroundCheckerColor1, DrawSettings.BackgroundCheckerColor2, DrawSettings.BackgroundCheckerSize
 	};
-	
+
 	if (CachedCheckerTextureSettings != NewCheckerTextureSettings)
 	{
 		if (CheckerTexture)
 		{
 		    DestroyCheckerTexture(CheckerTexture);
         }
-        
+
         if (NewCheckerTextureSettings.bEnabled)
 		{
 			CheckerTexture = TStrongObjectPtr<UTexture2D>(CreateCheckerTexture(NewCheckerTextureSettings.Color1.ToFColorSRGB(),
@@ -448,17 +466,17 @@ void FImageViewportClient::CreateOrDestroyCheckerTextureIfSettingsChanged(const
 	}
 }
 
-FVector2d FImageViewportClient::GetViewportSizeWithDPIScaling() const
+FVector2d FImageViewportClient::GetScaledViewportSize() const
 {
-	const FIntPoint ViewportSize = Viewport->GetSizeXY();
-	const float DPIScaleFactor = GetDPIScaleFactor.Execute();
-
-	return FVector2d(ViewportSize) / DPIScaleFactor;
+	// Return unzoomed logical size for controller Reset/SetZoom calls
+	return FVector2d(Viewport->GetSizeXY()) / GetLayoutAndDPIScaleFactor.Execute();
 }
 
 bool FImageViewportClient::MouseIsOverABComparisonDivider(const FIntPoint MousePos) const
 {
-	const double ABComparisonDividerPosition = CachedPlacement.Offset.X + CachedPlacement.Size.X * ABComparisonDivider;
+	// Scale the divider position from unzoomed logical space to RT pixel space for comparison with mouse pos
+	const float LayoutScale = GetLayoutScaleFactor.Execute();
+	const double ABComparisonDividerPosition = (CachedPlacement.Offset.X + CachedPlacement.Size.X * ABComparisonDivider) * LayoutScale;
 	return ABComparisonDividerPosition - 1.0 <= MousePos.X && MousePos.X <= ABComparisonDividerPosition + 1.0;
 }
 }
diff --git a/Engine/Plugins/Experimental/ImageWidgets/Source/ImageWidgets/Private/ImageViewportClient.h b/Engine/Plugins/Experimental/ImageWidgets/Source/ImageWidgets/Private/ImageViewportClient.h
--- a/Engine/Plugins/Experimental/ImageWidgets/Source/ImageWidgets/Private/ImageViewportClient.h
+++ b/Engine/Plugins/Experimental/ImageWidgets/Source/ImageWidgets/Private/ImageViewportClient.h
@@ -14,7 +14,8 @@ namespace UE::ImageWidgets
 	DECLARE_DELEGATE_RetVal(FIntPoint, FGetImageSize)
 	DECLARE_DELEGATE_ThreeParams(FDrawImage, FViewport*, FCanvas*, const IImageViewer::FDrawProperties&)
 	DECLARE_DELEGATE_RetVal(SImageViewport::FDrawSettings, FGetDrawSettings)
-	DECLARE_DELEGATE_RetVal(float, FGetDPIScaleFactor)
+	DECLARE_DELEGATE_RetVal(float, FGetLayoutAndDPIScaleFactor)
+	DECLARE_DELEGATE_RetVal(float, FGetLayoutScaleFactor)
 
 	/**
 	 * Viewport client for controlling the camera and drawing viewport contents. 
@@ -23,8 +24,9 @@ namespace UE::ImageWidgets
 	{
 	public:
 		FImageViewportClient(const TWeakPtr<SEditorViewport>& InEditorViewport, FGetImageSize&& InGetImageSize, FDrawImage&& InDrawImage,
-		                     FGetDrawSettings&& InGetDrawSettings, FGetDPIScaleFactor&& InGetDPIScaleFactor,
-		                     const FImageABComparison* InABComparison, const SImageViewport::FControllerSettings& InControllerSettings);
+		                     FGetDrawSettings&& InGetDrawSettings, FGetLayoutAndDPIScaleFactor&& InGetLayoutAndDPIScaleFactor,
+		                     FGetLayoutScaleFactor&& InGetLayoutScaleFactor, const FImageABComparison* InABComparison,
+		                     const SImageViewport::FControllerSettings& InControllerSettings);
 		virtual ~FImageViewportClient() override;
 
 		virtual void Draw(FViewport* InViewport, FCanvas* Canvas) override;
@@ -56,22 +58,23 @@ namespace UE::ImageWidgets
 			uint32 CheckerSize = 0;
 		};
 
-		FVector2d GetCurrentDragWithDPIScaling() const;
-		IImageViewer::FDrawProperties::FPlacement GetPlacementProperties(FIntPoint ImageSize, FVector2d ViewportSizeWithDPIScaling) const;
+		FVector2d GetCurrentDrag() const;
+		IImageViewer::FDrawProperties::FPlacement GetPlacementProperties(FIntPoint ImageSize, FVector2d ScaledViewportSize) const;
 		IImageViewer::FDrawProperties::FMip GetMipProperties() const;
 
 		void CreateOrDestroyCheckerTextureIfSettingsChanged(const SImageViewport::FDrawSettings& DrawSettings);
 		bool MouseIsOverABComparisonDivider(FIntPoint MousePos) const;
 
-		FVector2d GetViewportSizeWithDPIScaling() const;
-		
+		FVector2d GetScaledViewportSize() const;
+
 		FGetImageSize GetImageSize;
 		FDrawImage DrawImage;
 		FGetDrawSettings GetDrawSettings;
-		FGetDPIScaleFactor GetDPIScaleFactor;
+		FGetLayoutAndDPIScaleFactor GetLayoutAndDPIScaleFactor;
+		FGetLayoutScaleFactor GetLayoutScaleFactor;
 		SImageViewport::FControllerSettings::FOnInputKey OnInputKey;
 		const FImageABComparison* ABComparison;
-		
+
 		bool bDragging = false;
 		FIntPoint DraggingStart;
 
diff --git a/Engine/Plugins/Experimental/ImageWidgets/Source/ImageWidgets/Private/SImageViewport.cpp b/Engine/Plugins/Experimental/ImageWidgets/Source/ImageWidgets/Private/SImageViewport.cpp
--- a/Engine/Plugins/Experimental/ImageWidgets/Source/ImageWidgets/Private/SImageViewport.cpp
+++ b/Engine/Plugins/Experimental/ImageWidgets/Source/ImageWidgets/Private/SImageViewport.cpp
@@ -181,7 +181,8 @@ namespace UE::ImageWidgets
 			FGetImageSize::CreateLambda(GetImageSize),
 			FDrawImage::CreateLambda(DrawImage),
 			FGetDrawSettings::CreateLambda(GetDrawSettings),
-			FGetDPIScaleFactor::CreateRaw(this, &SImageViewport::GetDPIScaleFactor),
+			FGetLayoutAndDPIScaleFactor::CreateRaw(this, &SImageViewport::GetLayoutAndDPIScaleFactor),
+			FGetLayoutScaleFactor::CreateRaw(this, &SImageViewport::GetLayoutScaleFactor),
 			ABComparison.Get(),
 			InArgs._ControllerSettings)
 			);
@@ -289,7 +290,7 @@ namespace UE::ImageWidgets
 			           SImageViewportToolbar::FHasImage::CreateLambda(HasImage),
 			           SImageViewportToolbar::FNumMips::CreateLambda(GetNumMips),
 			           SImageViewportToolbar::FImageGuid::CreateLambda(GetImageGuid),
-			           FGetDPIScaleFactor::CreateRaw(this, &SImageViewport::GetDPIScaleFactor),
+			           FGetLayoutAndDPIScaleFactor::CreateRaw(this, &SImageViewport::GetLayoutAndDPIScaleFactor),
 			           SImageViewportToolbar::FGetOverlaySettings::CreateLambda(GetOverlaySettings),
 			           bABComparisonEnabled ? ABComparison.Get() : nullptr,
 			           ToolbarExtender
@@ -307,7 +308,7 @@ namespace UE::ImageWidgets
 			[
 				MakeViewportToolbarOverlay()
 			];
-	
+
 		Overlay->AddSlot()
 		       .VAlign(VAlign_Bottom)
 		[
@@ -315,6 +316,13 @@ namespace UE::ImageWidgets
 		];
 	}
 
+	void SImageViewport::OnArrangeChildren(const FGeometry& AllottedGeometry, FArrangedChildren& ArrangedChildren) const
+	{
+		SEditorViewport::OnArrangeChildren(AllottedGeometry, ArrangedChildren);
+
+		AllottedGeometryScale = AllottedGeometry.Scale;
+	}
+
 	float SImageViewport::GetDPIScaleFactor() const
 	{
 		if (const TSharedPtr<SWindow> TopLevelWindow = FSlateApplication::Get().FindWidgetWindow(SharedThis(this)); TopLevelWindow.IsValid())
@@ -324,6 +332,21 @@ namespace UE::ImageWidgets
 		return 1.0f;
 	}
 
+	float SImageViewport::GetLayoutScaleFactor() const
+	{
+		const float DPI = GetDPIScaleFactor();
+		if (DPI > SMALL_NUMBER)
+		{
+			return FMath::Max(AllottedGeometryScale / DPI, SMALL_NUMBER);
+		}
+		return 1.0f;
+	}
+
+	float SImageViewport::GetLayoutAndDPIScaleFactor() const
+	{
+		return FMath::Max(AllottedGeometryScale, SMALL_NUMBER);
+	}
+
 	FText SImageViewport::GetPickerLabel() const
 	{
 		const IImageViewer::FImageInfo ImageInfo = ImageViewer->GetCurrentImageInfo();
diff --git a/Engine/Plugins/Experimental/ImageWidgets/Source/ImageWidgets/Private/SImageViewportToolbar.cpp b/Engine/Plugins/Experimental/ImageWidgets/Source/ImageWidgets/Private/SImageViewportToolbar.cpp
--- a/Engine/Plugins/Experimental/ImageWidgets/Source/ImageWidgets/Private/SImageViewportToolbar.cpp
+++ b/Engine/Plugins/Experimental/ImageWidgets/Source/ImageWidgets/Private/SImageViewportToolbar.cpp
@@ -7,7 +7,6 @@
 #include "ImageWidgetsCommands.h"
 #include "SEditorViewportToolBarMenu.h"
 #include "Framework/MultiBox/MultiBoxBuilder.h"
-#include "Settings/LevelEditorViewportSettings.h"
 #include "Widgets/SBoxPanel.h"
 #include "Widgets/Input/SCheckBox.h"
 #include "Widgets/Layout/SBorder.h"
@@ -32,12 +31,12 @@ namespace UE::ImageWidgets
 		HasImage = MoveTemp(Parameters.HasImage);
 		NumMips = MoveTemp(Parameters.NumMips);
 		ImageGuid = MoveTemp(Parameters.ImageGuid);
-		GetDPIScaleFactor = MoveTemp(Parameters.GetDPIScaleFactor);
+		GetLayoutAndDPIScaleFactor = MoveTemp(Parameters.GetLayoutAndDPIScaleFactor);
 		GetOverlaySettings = MoveTemp(Parameters.GetOverlaySettings);
 		check(HasImage.IsBound());
 		check(NumMips.IsBound());
 		check(ImageGuid.IsBound());
-		check(GetDPIScaleFactor.IsBound());
+		check(GetLayoutAndDPIScaleFactor.IsBound());
 		check(GetOverlaySettings.IsBound());
 
 		ABComparison = Parameters.ABComparison;
@@ -137,7 +136,7 @@ namespace UE::ImageWidgets
 							{
 								return ABComparison->CanSetABComparison(AorB);
 							})
-						.IsChecked(ABComparison->ABComparisonIsSet(AorB))
+						.IsChecked(ABComparison->ABComparisonIsSet(AorB) ? ECheckBoxState::Checked : ECheckBoxState::Unchecked)
 						.OnCheckStateChanged_Lambda([this, AorB](const ECheckBoxState State)
 						    {
 								ABComparison->SetABComparison(AorB, State != ECheckBoxState::Checked ? FGuid() : ImageGuid.Execute());
@@ -200,9 +199,9 @@ namespace UE::ImageWidgets
 	{
 		const FImageViewportController::FZoomSettings ZoomSettings = ViewportClient->GetZoom();
 
-		const float DPIScaleFactor = GetDPIScaleFactor.Execute();
+		const float LayoutAndDPIScaleFactor = GetLayoutAndDPIScaleFactor.Execute();
 
-		const double Zoom = ZoomSettings.Zoom / DPIScaleFactor;
+		const double Zoom = ZoomSettings.Zoom / LayoutAndDPIScaleFactor;
 		FNumberFormattingOptions FormattingOptions;
 		FormattingOptions.SetMaximumFractionalDigits(Zoom < 1.0 ? (Zoom < 0.1 ? 2 : 1) : 0);
 		const FText ZoomPercentage = FText::AsPercent(Zoom, &FormattingOptions);
diff --git a/Engine/Plugins/Experimental/ImageWidgets/Source/ImageWidgets/Private/SImageViewportToolbar.h b/Engine/Plugins/Experimental/ImageWidgets/Source/ImageWidgets/Private/SImageViewportToolbar.h
--- a/Engine/Plugins/Experimental/ImageWidgets/Source/ImageWidgets/Private/SImageViewportToolbar.h
+++ b/Engine/Plugins/Experimental/ImageWidgets/Source/ImageWidgets/Private/SImageViewportToolbar.h
@@ -38,12 +38,12 @@ namespace UE::ImageWidgets
 			FHasImage HasImage;
 			FNumMips NumMips;
 			FImageGuid ImageGuid;
-			FGetDPIScaleFactor GetDPIScaleFactor;
+			FGetLayoutAndDPIScaleFactor GetLayoutAndDPIScaleFactor;
 			FGetOverlaySettings GetOverlaySettings;
 			FImageABComparison* ABComparison = nullptr;
 			TSharedPtr<FExtender> ToolbarExtender;
 		};
-		
+
 		void Construct(const FArguments& InArgs, const TSharedPtr<FImageViewportClient>& InViewportClient, const TSharedPtr<FUICommandList>& InCommandList,
 		               FConstructParameters Parameters);
 
@@ -67,7 +67,7 @@ namespace UE::ImageWidgets
 		FHasImage HasImage;
 		FNumMips NumMips;
 		FImageGuid ImageGuid;
-		FGetDPIScaleFactor GetDPIScaleFactor;
+		FGetLayoutAndDPIScaleFactor GetLayoutAndDPIScaleFactor;
 		FGetOverlaySettings GetOverlaySettings;
 		FImageABComparison* ABComparison = nullptr;
 	};
diff --git a/Engine/Plugins/Experimental/ImageWidgets/Source/ImageWidgets/Public/SImageViewport.h b/Engine/Plugins/Experimental/ImageWidgets/Source/ImageWidgets/Public/SImageViewport.h
--- a/Engine/Plugins/Experimental/ImageWidgets/Source/ImageWidgets/Public/SImageViewport.h
+++ b/Engine/Plugins/Experimental/ImageWidgets/Source/ImageWidgets/Public/SImageViewport.h
@@ -230,22 +230,32 @@ namespace UE::ImageWidgets
 
 		// SEditorViewport overrides - begin
 	public:
-		virtual void BindCommands() override;
 		virtual bool IsVisible() const override { return true; } // Overriding this method avoids the 0.25 seconds update delay in SEditorViewport.
 	protected:
+		virtual void BindCommands() override;
 		virtual TSharedRef<FEditorViewportClient> MakeEditorViewportClient() override;
 		virtual void PopulateViewportOverlays(TSharedRef<SOverlay> Overlay) override;
 		// SEditorViewport overrides - end
 
+		// SCompoundWidget overrides - begin
+		virtual void OnArrangeChildren(const FGeometry& AllottedGeometry, FArrangedChildren& ArrangedChildren) const override;
+		// SCompoundWidget overrides - end
+
 	private:
 		TSharedRef<SWidget> MakeViewportToolbarOverlay();
-	
+
 		/** Create the status bar widgets. */
 		TSharedRef<SWidget> MakeStatusBar(TSharedRef<SOverlay> Overlay);
 
-		/** Provides the UI scale factor. */
+		/** Provides the window DPI scale factor. */
 		float GetDPIScaleFactor() const;
 
+		/** Provides the external layout factor (ratio of accumulated scale to window DPI). */
+		float GetLayoutScaleFactor() const;
+
+		/** Provides the layout and DPI scale factor, i.e. the full accumulated geometry scale. */
+		float GetLayoutAndDPIScaleFactor() const;
+
 		/** Provides the text for the color picker display in the status bar. */
 		FText GetPickerLabel() const;
 
@@ -278,5 +288,8 @@ namespace UE::ImageWidgets
 
 		/** Status bar extensions provided by the call to @see Construct(). This pointer is reset after the extensions were applied during construction. */
 		TSharedPtr<FStatusBarExtender> StatusBarExtender;
+
+		/** Cached scale for the allotted geometry; used to correctly incorporate the overall scale of the layout when drawing the image. */
+		mutable float AllottedGeometryScale = 1.0f;
 	};
 }
