ImageWidgets: Add new resize viewport behavior that adjusts zoom with viewport resize

#rn ImageWidgets: Added new controler setting to adjust zoom when resizing
#Synced-CL 51725222

diff --git a/Engine/Plugins/Experimental/ImageWidgets/Source/ImageWidgets/Private/ImageViewportClient.cpp b/Engine/Plugins/Experimental/ImageWidgets/Source/ImageWidgets/Private/ImageViewportClient.cpp
--- a/Engine/Plugins/Experimental/ImageWidgets/Source/ImageWidgets/Private/ImageViewportClient.cpp
+++ b/Engine/Plugins/Experimental/ImageWidgets/Source/ImageWidgets/Private/ImageViewportClient.cpp
@@ -111,6 +111,7 @@ FImageViewportClient::FImageViewportClient(const TWeakPtr<SEditorViewport>& InEd
 	, OnInputKey(InControllerSettings.OnInputKey)
 	, ABComparison(InABComparison)
 	, Controller(static_cast<FImageViewportController::EZoomMode>(InControllerSettings.DefaultZoomMode))
+	, bZoomOnResize(InControllerSettings.bZoomOnResize)
 {
 	check(GetImageSize.IsBound());
 	check(DrawImage.IsBound());
@@ -143,6 +144,47 @@ void FImageViewportClient::Draw(FViewport* InViewport, FCanvas* Canvas)
 	}
 
 	const float LayoutAndDPIScale = GetLayoutAndDPIScaleFactor.Execute();
+
+	// Recalculate zoom and pan when the viewport is resized.
+	{
+		const FVector2d ScaledSize = FVector2d(InViewport->GetSizeXY()) / LayoutAndDPIScale;
+		if (ScaledSize != CachedScaledViewportSize)
+		{
+			if (CachedScaledViewportSize != FVector2d::ZeroVector)
+			{
+				const FImageViewportController::FZoomSettings OldZoomSettings = Controller.GetZoom();
+				const FVector2d OldPan = Controller.GetPan(FVector2d::Zero());
+
+				if (OldZoomSettings.Mode != FImageViewportController::EZoomMode::Custom)
+				{
+					// Fill/Fit: recalculate zoom for new viewport size.
+					Controller.SetZoom(OldZoomSettings.Mode, OldZoomSettings.Zoom, ImageSize, ScaledSize);
+				}
+				else if (bZoomOnResize)
+				{
+					// When zoom on resize is set, scale zoom proportionally.
+					const double OldFillZoom = (CachedScaledViewportSize / FVector2d(ImageSize)).GetMin();
+					const double NewFillZoom = (ScaledSize / FVector2d(ImageSize)).GetMin();
+					if (OldFillZoom > SMALL_NUMBER)
+					{
+						Controller.SetZoom(FImageViewportController::EZoomMode::Custom, OldZoomSettings.Zoom * (NewFillZoom / OldFillZoom),
+						                   ImageSize, ScaledSize);
+					}
+				}
+
+				// Scale pan uniformly by the zoom ratio to maintain relative image position.
+				// Using a single scalar (NewZoom/OldZoom) for both axes avoids diagonal drift when the constraining axis flips, e.g. landscape to portrait.
+				const double NewZoom = Controller.GetZoom().Zoom;
+				const double PanScale = OldZoomSettings.Zoom > SMALL_NUMBER ? NewZoom / OldZoomSettings.Zoom : 1.0;
+				const FVector2d NewPan = OldPan * PanScale;
+				const FVector2d CurrentPan = Controller.GetPan(FVector2d::Zero());
+				Controller.Pan(NewPan - CurrentPan);
+			}
+
+			CachedScaledViewportSize = ScaledSize;
+		}
+	}
+
 	const float LayoutScale = GetLayoutScaleFactor.Execute();
 
 	// Controller computes placement in unzoomed logical space
diff --git a/Engine/Plugins/Experimental/ImageWidgets/Source/ImageWidgets/Private/ImageViewportClient.h b/Engine/Plugins/Experimental/ImageWidgets/Source/ImageWidgets/Private/ImageViewportClient.h
--- a/Engine/Plugins/Experimental/ImageWidgets/Source/ImageWidgets/Private/ImageViewportClient.h
+++ b/Engine/Plugins/Experimental/ImageWidgets/Source/ImageWidgets/Private/ImageViewportClient.h
@@ -89,6 +89,9 @@ namespace UE::ImageWidgets
 		TStrongObjectPtr<UTexture2D> CheckerTexture;
 		FCheckerTextureSettings CachedCheckerTextureSettings;
 
+		FVector2d CachedScaledViewportSize = FVector2d::ZeroVector;
+		bool bZoomOnResize = false;
+
 		double ABComparisonDivider = 0.5;
 		bool bDraggingABComparisonDivider = false;
 	};
diff --git a/Engine/Plugins/Experimental/ImageWidgets/Source/ImageWidgets/Public/SImageViewport.h b/Engine/Plugins/Experimental/ImageWidgets/Source/ImageWidgets/Public/SImageViewport.h
--- a/Engine/Plugins/Experimental/ImageWidgets/Source/ImageWidgets/Public/SImageViewport.h
+++ b/Engine/Plugins/Experimental/ImageWidgets/Source/ImageWidgets/Public/SImageViewport.h
@@ -122,6 +122,9 @@ namespace UE::ImageWidgets
 
 			/** Delegate for custom input key handling. */
 			FOnInputKey OnInputKey;
+
+			/** Adjusts zoom to be proportional to viewport resizes, e.g. doubling the viewport size will double the zoom. */
+			bool bZoomOnResize = false;
 		};
 
 		/**
