ImageWidgets: Arguments for SImageViewport to configure toolbar style

#Synced-CL 52046073

diff --git a/Engine/Plugins/Experimental/ImageWidgets/Source/ImageWidgets/Private/SImageViewport.cpp b/Engine/Plugins/Experimental/ImageWidgets/Source/ImageWidgets/Private/SImageViewport.cpp
--- a/Engine/Plugins/Experimental/ImageWidgets/Source/ImageWidgets/Private/SImageViewport.cpp
+++ b/Engine/Plugins/Experimental/ImageWidgets/Source/ImageWidgets/Private/SImageViewport.cpp
@@ -150,6 +150,9 @@ namespace UE::ImageWidgets
 		check(ImageViewer.IsValid());
 
 		ToolbarExtender = InArgs._ToolbarExtender;
+		ToolbarExtenderStyle = InArgs._ToolbarExtenderStyle;
+		ToolbarExtenderStyleName = InArgs._ToolbarExtenderStyleName;
+
 		StatusBarExtender = InArgs._StatusBarExtender;
 
 		DrawSettings = InArgs._DrawSettings;
@@ -293,10 +296,12 @@ namespace UE::ImageWidgets
 			           FGetLayoutAndDPIScaleFactor::CreateRaw(this, &SImageViewport::GetLayoutAndDPIScaleFactor),
 			           SImageViewportToolbar::FGetOverlaySettings::CreateLambda(GetOverlaySettings),
 			           bABComparisonEnabled ? ABComparison.Get() : nullptr,
-			           ToolbarExtender
+			           ToolbarExtender, ToolbarExtenderStyle, ToolbarExtenderStyleName
 		           });
 
 		ToolbarExtender.Reset();
+		ToolbarExtenderStyle = nullptr;
+		ToolbarExtenderStyleName = FName();
 
 		return ImageViewportToolbar.ToSharedRef();
 	}
diff --git a/Engine/Plugins/Experimental/ImageWidgets/Source/ImageWidgets/Private/SImageViewportToolbar.cpp b/Engine/Plugins/Experimental/ImageWidgets/Source/ImageWidgets/Private/SImageViewportToolbar.cpp
--- a/Engine/Plugins/Experimental/ImageWidgets/Source/ImageWidgets/Private/SImageViewportToolbar.cpp
+++ b/Engine/Plugins/Experimental/ImageWidgets/Source/ImageWidgets/Private/SImageViewportToolbar.cpp
@@ -41,6 +41,9 @@ namespace UE::ImageWidgets
 
 		ABComparison = Parameters.ABComparison;
 
+		ToolbarExtenderStyle = Parameters.ToolbarExtenderStyle;
+		ToolbarExtenderStyleName = Parameters.ToolbarExtenderStyleName;
+
 		ChildSlot
 		[
 			SNew(SBorder)
@@ -71,18 +74,23 @@ namespace UE::ImageWidgets
 		SViewportToolBar::Construct(SViewportToolBar::FArguments());
 	}
 
-	FSlimHorizontalToolBarBuilder GetToolbarBuilder(const TSharedPtr<FUICommandList>& CommandList, const TSharedPtr<FExtender>& Extender)
+	FSlimHorizontalToolBarBuilder SImageViewportToolbar::GetToolbarBuilder(const TSharedPtr<FExtender>& Extender) const
 	{
 		FSlimHorizontalToolBarBuilder ToolbarBuilder(CommandList, FMultiBoxCustomization::None, Extender, true);
-		ToolbarBuilder.SetStyle(&FAppStyle::Get(), "EditorViewportToolBar");
+
+		const ISlateStyle* Style = ToolbarExtenderStyle ? ToolbarExtenderStyle : &FAppStyle::Get();
+		const FName StyleName = ToolbarExtenderStyle && !ToolbarExtenderStyleName.IsNone() ? ToolbarExtenderStyleName : FName("EditorViewportToolBar");
+		ToolbarBuilder.SetStyle(Style, StyleName);
+
 		ToolbarBuilder.SetIsFocusable(false);
 		ToolbarBuilder.SetLabelVisibility(EVisibility::Collapsed);
+
 		return ToolbarBuilder;
-	};
+	}
 
 	TSharedRef<SWidget> SImageViewportToolbar::MakeLeftToolbar(const TSharedPtr<FExtender>& Extender)
 	{
-		FSlimHorizontalToolBarBuilder ToolbarBuilder = GetToolbarBuilder(CommandList, Extender);
+		FSlimHorizontalToolBarBuilder ToolbarBuilder = GetToolbarBuilder(Extender);
 
 		ToolbarBuilder.BeginSection("ToolbarLeft");
 		{
@@ -119,7 +127,7 @@ namespace UE::ImageWidgets
 
 	TSharedRef<SWidget> SImageViewportToolbar::MakeCenterToolbar(const TSharedPtr<FExtender>& Extender, bool bEnableABComparison)
 	{
-		FSlimHorizontalToolBarBuilder ToolbarBuilder = GetToolbarBuilder(CommandList, Extender);
+		FSlimHorizontalToolBarBuilder ToolbarBuilder = GetToolbarBuilder(Extender);
 
 		ToolbarBuilder.BeginSection("ToolbarCenter");
 		{
@@ -182,7 +190,7 @@ namespace UE::ImageWidgets
 
 	TSharedRef<SWidget> SImageViewportToolbar::MakeRightToolbar(const TSharedPtr<FExtender>& Extender)
 	{
-		FSlimHorizontalToolBarBuilder ToolbarBuilder = GetToolbarBuilder(CommandList, Extender);
+		FSlimHorizontalToolBarBuilder ToolbarBuilder = GetToolbarBuilder(Extender);
 
 		ToolbarBuilder.BeginSection("ToolbarRight");
 		{
diff --git a/Engine/Plugins/Experimental/ImageWidgets/Source/ImageWidgets/Private/SImageViewportToolbar.h b/Engine/Plugins/Experimental/ImageWidgets/Source/ImageWidgets/Private/SImageViewportToolbar.h
--- a/Engine/Plugins/Experimental/ImageWidgets/Source/ImageWidgets/Private/SImageViewportToolbar.h
+++ b/Engine/Plugins/Experimental/ImageWidgets/Source/ImageWidgets/Private/SImageViewportToolbar.h
@@ -9,6 +9,7 @@
 
 class FExtender;
 class FImageViewportCommands;
+class FSlimHorizontalToolBarBuilder;
 class FUICommandList;
 class SViewportToolBarComboMenu;
 
@@ -42,12 +43,15 @@ namespace UE::ImageWidgets
 			FGetOverlaySettings GetOverlaySettings;
 			FImageABComparison* ABComparison = nullptr;
 			TSharedPtr<FExtender> ToolbarExtender;
+			const ISlateStyle* ToolbarExtenderStyle = nullptr;
+			FName ToolbarExtenderStyleName;
 		};
 
 		void Construct(const FArguments& InArgs, const TSharedPtr<FImageViewportClient>& InViewportClient, const TSharedPtr<FUICommandList>& InCommandList,
 		               FConstructParameters Parameters);
 
 	private:
+		FSlimHorizontalToolBarBuilder GetToolbarBuilder(const TSharedPtr<FExtender>& Extender) const;
 		TSharedRef<SWidget> MakeLeftToolbar(const TSharedPtr<FExtender>& Extender);
 		TSharedRef<SWidget> MakeCenterToolbar(const TSharedPtr<FExtender>& Extender, bool bEnableABComparison);
 		TSharedRef<SWidget> MakeRightToolbar(const TSharedPtr<FExtender>& Extender);
@@ -70,5 +74,7 @@ namespace UE::ImageWidgets
 		FGetLayoutAndDPIScaleFactor GetLayoutAndDPIScaleFactor;
 		FGetOverlaySettings GetOverlaySettings;
 		FImageABComparison* ABComparison = nullptr;
+		const ISlateStyle* ToolbarExtenderStyle = nullptr;
+		FName ToolbarExtenderStyleName;
 	};
 }
diff --git a/Engine/Plugins/Experimental/ImageWidgets/Source/ImageWidgets/Public/SImageViewport.h b/Engine/Plugins/Experimental/ImageWidgets/Source/ImageWidgets/Public/SImageViewport.h
--- a/Engine/Plugins/Experimental/ImageWidgets/Source/ImageWidgets/Public/SImageViewport.h
+++ b/Engine/Plugins/Experimental/ImageWidgets/Source/ImageWidgets/Public/SImageViewport.h
@@ -156,6 +156,15 @@ namespace UE::ImageWidgets
 			/** Extensions for the viewport toolbar; valid hooks are "ToolbarLeft", "ToolbarCenter", "ToolbarRight" */
 			SLATE_ARGUMENT(TSharedPtr<FExtender>, ToolbarExtender)
 
+			/** Style override for viewport toolbar extensions. */
+			SLATE_ARGUMENT(const ISlateStyle*, ToolbarExtenderStyle)
+
+			/**
+			 * Name used for the viewport toolbar extensions style override.
+			 * Only used if ToolbarExtenderStyle is set. If this is not set, "EditorViewportToolBar" will be used.
+			 */
+			SLATE_ARGUMENT(FName, ToolbarExtenderStyleName)
+
 			/** Extensions for the viewport status bar; valid hooks are "StatusBarLeft", "StatusBarCenter", "StatusBarRight" */
 			SLATE_ARGUMENT(TSharedPtr<FStatusBarExtender>, StatusBarExtender)
 
@@ -289,6 +298,10 @@ namespace UE::ImageWidgets
 		/** Toolbar extensions provided by the call to @see Construct(). This pointer is reset after the extensions were applied during construction. */
 		TSharedPtr<FExtender> ToolbarExtender;
 
+		/** Style override for toolbar extender. Only used during construction, then reset. */
+		const ISlateStyle* ToolbarExtenderStyle = nullptr;
+		FName ToolbarExtenderStyleName;
+
 		/** Status bar extensions provided by the call to @see Construct(). This pointer is reset after the extensions were applied during construction. */
 		TSharedPtr<FStatusBarExtender> StatusBarExtender;
 
