ImageWidgets: Fix mip selection bug

#Synced-CL 52049184

#Synced-CL 52049184

diff --git a/Engine/Plugins/Experimental/ImageWidgets/Source/ImageWidgets/Private/ImageViewportClient.h b/Engine/Plugins/Experimental/ImageWidgets/Source/ImageWidgets/Private/ImageViewportClient.h
--- a/Engine/Plugins/Experimental/ImageWidgets/Source/ImageWidgets/Private/ImageViewportClient.h
+++ b/Engine/Plugins/Experimental/ImageWidgets/Source/ImageWidgets/Private/ImageViewportClient.h
@@ -79,7 +79,7 @@ namespace UE::ImageWidgets
 		FIntPoint DraggingStart;
 
 		bool bMipSelected = true;
-		int32 MipLevel = -1;
+		int32 MipLevel = INDEX_NONE;
 
 		IImageViewer::FDrawProperties::FPlacement CachedPlacement;
 		bool bCachedPlacementIsValid = false;
diff --git a/Engine/Plugins/Experimental/ImageWidgets/Source/ImageWidgets/Private/SImageViewport.cpp b/Engine/Plugins/Experimental/ImageWidgets/Source/ImageWidgets/Private/SImageViewport.cpp
--- a/Engine/Plugins/Experimental/ImageWidgets/Source/ImageWidgets/Private/SImageViewport.cpp
+++ b/Engine/Plugins/Experimental/ImageWidgets/Source/ImageWidgets/Private/SImageViewport.cpp
@@ -124,7 +124,7 @@ namespace UE::ImageWidgets
 					}
 				}
 
-				if (CurrentMip <= NumMips)
+				if (NumMips > 0 && CurrentMip < NumMips - 1)
 				{
 					ViewportClient->SetMipLevel(CurrentMip + 1);
 				}
@@ -135,7 +135,7 @@ namespace UE::ImageWidgets
 			{
 				const TSharedPtr<FImageViewportClient> ViewportClient = StaticCastSharedPtr<FImageViewportClient>(Client);
 				const int32 CurrentMip = ViewportClient->GetMipLevel();
-				if (CurrentMip > -1)
+				if (CurrentMip > INDEX_NONE)
 				{
 					ViewportClient->SetMipLevel(CurrentMip - 1);
 				}
@@ -201,7 +201,7 @@ namespace UE::ImageWidgets
 	void SImageViewport::ResetMip() const
 	{
 		const int32 SelectedMip = ImageViewportClient->GetMipLevel();
-		if (SelectedMip != -1)
+		if (SelectedMip != INDEX_NONE)
 		{
 			if (ImageViewer.IsValid())
 			{
diff --git a/Engine/Plugins/Experimental/ImageWidgets/Source/ImageWidgets/Private/SImageViewportToolbar.cpp b/Engine/Plugins/Experimental/ImageWidgets/Source/ImageWidgets/Private/SImageViewportToolbar.cpp
--- a/Engine/Plugins/Experimental/ImageWidgets/Source/ImageWidgets/Private/SImageViewportToolbar.cpp
+++ b/Engine/Plugins/Experimental/ImageWidgets/Source/ImageWidgets/Private/SImageViewportToolbar.cpp
@@ -263,7 +263,7 @@ namespace UE::ImageWidgets
 	FText SImageViewportToolbar::GetMipMenuLabel() const
 	{
 		const int32 Mip = ViewportClient->GetMipLevel();
-		return FText::Format(LOCTEXT("Mip", "Mip {0}"), Mip == -1 ? ImageViewportToolbar_Local::Auto : FText::AsNumber(Mip));
+		return FText::Format(LOCTEXT("Mip", "Mip {0}"), Mip == INDEX_NONE ? ImageViewportToolbar_Local::Auto : FText::AsNumber(Mip));
 	}
 
 	TSharedRef<SWidget> SImageViewportToolbar::MakeMipMenu() const
@@ -286,7 +286,7 @@ namespace UE::ImageWidgets
 			);
 		};
 
-		AddMenuEntry(ImageViewportToolbar_Local::Auto, LOCTEXT("MipLevelAuto_Tooltip", "Choose Mip Level automatically"), -1);
+		AddMenuEntry(ImageViewportToolbar_Local::Auto, LOCTEXT("MipLevelAuto_Tooltip", "Choose Mip Level automatically"), INDEX_NONE);
 
 		if (NumMips.IsBound())
 		{
