ImageWidgets: clean-up

#Synced-CL 52049184

diff --git a/Engine/Plugins/Experimental/ImageWidgets/Source/ImageWidgets/ImageWidgets.Build.cs b/Engine/Plugins/Experimental/ImageWidgets/Source/ImageWidgets/ImageWidgets.Build.cs
--- a/Engine/Plugins/Experimental/ImageWidgets/Source/ImageWidgets/ImageWidgets.Build.cs
+++ b/Engine/Plugins/Experimental/ImageWidgets/Source/ImageWidgets/ImageWidgets.Build.cs
@@ -12,34 +12,31 @@ public class ImageWidgets : ModuleRules
 		PCHUsage = PCHUsageMode.UseExplicitOrSharedPCHs;
 
 		PublicDependencyModuleNames.AddRange(
-			new[]
-			{
+			[
 				"Core",
 				"Slate",
+				"SlateCore",
 				"UnrealEd"
-			}
+			]
 		);
 
 		PrivateDependencyModuleNames.AddRange(
-			new[]
-			{
+			[
 				"ApplicationCore",
 				"CoreUObject",
 				"Engine",
-				"InputCore",
-				"SlateCore"
-			}
+				"InputCore"
+			]
 		);
 
 		bool bBuildColorViewerSample = ForceBuildColorViewerSample || Target.Configuration == UnrealTargetConfiguration.Debug;
 		if (bBuildColorViewerSample)
 		{
 			PrivateDependencyModuleNames.AddRange(
-				new[]
-				{
+				[
 					"Projects",
 					"WorkspaceMenuStructure"
-				}
+				]
 			);
 		}
 		PublicDefinitions.Add(bBuildColorViewerSample ? "IMAGE_WIDGETS_BUILD_COLOR_VIEWER_SAMPLE=1" : "IMAGE_WIDGETS_BUILD_COLOR_VIEWER_SAMPLE=0");
diff --git a/Engine/Plugins/Experimental/ImageWidgets/Source/ImageWidgets/Private/ColorViewerSample/ColorViewer.cpp b/Engine/Plugins/Experimental/ImageWidgets/Source/ImageWidgets/Private/ColorViewerSample/ColorViewer.cpp
--- a/Engine/Plugins/Experimental/ImageWidgets/Source/ImageWidgets/Private/ColorViewerSample/ColorViewer.cpp
+++ b/Engine/Plugins/Experimental/ImageWidgets/Source/ImageWidgets/Private/ColorViewerSample/ColorViewer.cpp
@@ -4,8 +4,8 @@
 
 #include "ColorViewer.h"
 
-#include <CanvasItem.h>
-#include <CanvasTypes.h>
+#include "CanvasItem.h"
+#include "CanvasTypes.h"
 
 #define LOCTEXT_NAMESPACE "ColorViewer"
 
diff --git a/Engine/Plugins/Experimental/ImageWidgets/Source/ImageWidgets/Private/ColorViewerSample/ColorViewer.h b/Engine/Plugins/Experimental/ImageWidgets/Source/ImageWidgets/Private/ColorViewerSample/ColorViewer.h
--- a/Engine/Plugins/Experimental/ImageWidgets/Source/ImageWidgets/Private/ColorViewerSample/ColorViewer.h
+++ b/Engine/Plugins/Experimental/ImageWidgets/Source/ImageWidgets/Private/ColorViewerSample/ColorViewer.h
@@ -4,7 +4,7 @@
 
 #if IMAGE_WIDGETS_BUILD_COLOR_VIEWER_SAMPLE
 
-#include <IImageViewer.h>
+#include "IImageViewer.h"
 
 namespace UE::ImageWidgets::Sample
 {
diff --git a/Engine/Plugins/Experimental/ImageWidgets/Source/ImageWidgets/Private/ColorViewerSample/ColorViewerWidget.h b/Engine/Plugins/Experimental/ImageWidgets/Source/ImageWidgets/Private/ColorViewerSample/ColorViewerWidget.h
--- a/Engine/Plugins/Experimental/ImageWidgets/Source/ImageWidgets/Private/ColorViewerSample/ColorViewerWidget.h
+++ b/Engine/Plugins/Experimental/ImageWidgets/Source/ImageWidgets/Private/ColorViewerSample/ColorViewerWidget.h
@@ -6,9 +6,9 @@
 
 #include "ColorViewer.h"
 
-#include <SImageCatalog.h>
-#include <SImageViewport.h>
-#include <Widgets/Layout/SSplitter.h>
+#include "SImageCatalog.h"
+#include "SImageViewport.h"
+#include "Widgets/Layout/SSplitter.h"
 
 namespace UE::ImageWidgets::Sample
 {
diff --git a/Engine/Plugins/Experimental/ImageWidgets/Source/ImageWidgets/Private/ImageViewportClient.cpp b/Engine/Plugins/Experimental/ImageWidgets/Source/ImageWidgets/Private/ImageViewportClient.cpp
--- a/Engine/Plugins/Experimental/ImageWidgets/Source/ImageWidgets/Private/ImageViewportClient.cpp
+++ b/Engine/Plugins/Experimental/ImageWidgets/Source/ImageWidgets/Private/ImageViewportClient.cpp
@@ -232,7 +232,7 @@ void FImageViewportClient::Draw(FViewport* InViewport, FCanvas* Canvas)
 		Canvas->DrawItem(Background);
 	}
 
-	const double ABComparisonDividerX = [this, LayoutScale]
+	const double ABComparisonDividerX = [this, LayoutScale, &ImageSize]
 	{
 		if (bDraggingABComparisonDivider)
 		{
@@ -243,7 +243,6 @@ void FImageViewportClient::Draw(FViewport* InViewport, FCanvas* Canvas)
 			FVector2d DividerPos = FVector2d(MousePos) / LayoutScale - CachedPlacement.Offset;
 			DividerPos /= CachedPlacement.ZoomFactor;
 
-			const FIntPoint ImageSize = GetImageSize.Execute();
 			return DividerPos.X / ImageSize.X;
 		}
 		return ABComparisonDivider;
@@ -519,7 +518,8 @@ bool FImageViewportClient::MouseIsOverABComparisonDivider(const FIntPoint MouseP
 	// Scale the divider position from unzoomed logical space to RT pixel space for comparison with mouse pos
 	const float LayoutScale = GetLayoutScaleFactor.Execute();
 	const double ABComparisonDividerPosition = (CachedPlacement.Offset.X + CachedPlacement.Size.X * ABComparisonDivider) * LayoutScale;
-	return ABComparisonDividerPosition - 1.0 <= MousePos.X && MousePos.X <= ABComparisonDividerPosition + 1.0;
+	const double HitMargin = FMath::Max(4.0, 2.0 * LayoutScale);
+	return ABComparisonDividerPosition - HitMargin <= MousePos.X && MousePos.X <= ABComparisonDividerPosition + HitMargin;
 }
 }
 
diff --git a/Engine/Plugins/Experimental/ImageWidgets/Source/ImageWidgets/Private/ImageWidgetsLogCategory.h b/Engine/Plugins/Experimental/ImageWidgets/Source/ImageWidgets/Private/ImageWidgetsLogCategory.h
--- a/Engine/Plugins/Experimental/ImageWidgets/Source/ImageWidgets/Private/ImageWidgetsLogCategory.h
+++ b/Engine/Plugins/Experimental/ImageWidgets/Source/ImageWidgets/Private/ImageWidgetsLogCategory.h
@@ -2,7 +2,7 @@
 
 #pragma once
 
-#include <Logging/LogMacros.h>
+#include "Logging/LogMacros.h"
 
 namespace UE::ImageWidgets
 {
diff --git a/Engine/Plugins/Experimental/ImageWidgets/Source/ImageWidgets/Private/ImageWidgetsStyle.cpp b/Engine/Plugins/Experimental/ImageWidgets/Source/ImageWidgets/Private/ImageWidgetsStyle.cpp
--- a/Engine/Plugins/Experimental/ImageWidgets/Source/ImageWidgets/Private/ImageWidgetsStyle.cpp
+++ b/Engine/Plugins/Experimental/ImageWidgets/Source/ImageWidgets/Private/ImageWidgetsStyle.cpp
@@ -9,7 +9,7 @@
 
 namespace UE::ImageWidgets
 {
-	FName FImageWidgetsStyle::StyleName("ImageViewportStyle");
+	FName FImageWidgetsStyle::StyleName("ImageWidgetsStyle");
 
 	FImageWidgetsStyle::FImageWidgetsStyle()
 		: FSlateStyleSet(StyleName)
@@ -22,7 +22,6 @@ namespace UE::ImageWidgets
 		Set("RichTextBlock.Green", FTextBlockStyle(NormalText).SetColorAndOpacity(FLinearColor(0.1f, 1.0f, 0.1f)));
 		Set("RichTextBlock.Blue", FTextBlockStyle(NormalText).SetColorAndOpacity(FLinearColor(0.1f, 0.1f, 1.0f)));
 
-
 		FSlateBrush* BrushTableRowOdd = new FSlateBrush();
 		BrushTableRowOdd->TintColor = FLinearColor(0.5f, 0.5f, 0.5f, 1.0f);
 		Set("TableRowOdd", BrushTableRowOdd);
diff --git a/Engine/Plugins/Experimental/ImageWidgets/Source/ImageWidgets/Private/SImageCatalog.cpp b/Engine/Plugins/Experimental/ImageWidgets/Source/ImageWidgets/Private/SImageCatalog.cpp
--- a/Engine/Plugins/Experimental/ImageWidgets/Source/ImageWidgets/Private/SImageCatalog.cpp
+++ b/Engine/Plugins/Experimental/ImageWidgets/Source/ImageWidgets/Private/SImageCatalog.cpp
@@ -2,13 +2,13 @@
 
 #include "SImageCatalog.h"
 
-#include <Algo/Accumulate.h>
-#include <Brushes/SlateColorBrush.h>
-#include <Framework/Application/SlateApplication.h>
-#include <Styling/StyleColors.h>
-#include <Widgets/Layout/SScrollBox.h>
-#include <Widgets/Views/SListView.h>
-#include <Widgets/Input/SButton.h>
+#include "Algo/Accumulate.h"
+#include "Brushes/SlateColorBrush.h"
+#include "Framework/Application/SlateApplication.h"
+#include "Styling/StyleColors.h"
+#include "Widgets/Layout/SScrollBox.h"
+#include "Widgets/Views/SListView.h"
+#include "Widgets/Input/SButton.h"
 
 #include "SImageCatalogItem.h"
 
diff --git a/Engine/Plugins/Experimental/ImageWidgets/Source/ImageWidgets/Private/SImageCatalogItem.cpp b/Engine/Plugins/Experimental/ImageWidgets/Source/ImageWidgets/Private/SImageCatalogItem.cpp
--- a/Engine/Plugins/Experimental/ImageWidgets/Source/ImageWidgets/Private/SImageCatalogItem.cpp
+++ b/Engine/Plugins/Experimental/ImageWidgets/Source/ImageWidgets/Private/SImageCatalogItem.cpp
@@ -3,9 +3,9 @@
 #include "SImageCatalogItem.h"
 #include "SImageCatalog.h"
 
-#include <Widgets/SBoxPanel.h>
-#include <Widgets/Images/SImage.h>
-#include <Widgets/Text/STextBlock.h>
+#include "Widgets/SBoxPanel.h"
+#include "Widgets/Images/SImage.h"
+#include "Widgets/Text/STextBlock.h"
 
 #define LOCTEXT_NAMESPACE "SImageViewerCatalogItem"
 
diff --git a/Engine/Plugins/Experimental/ImageWidgets/Source/ImageWidgets/Private/SImageViewport.cpp b/Engine/Plugins/Experimental/ImageWidgets/Source/ImageWidgets/Private/SImageViewport.cpp
--- a/Engine/Plugins/Experimental/ImageWidgets/Source/ImageWidgets/Private/SImageViewport.cpp
+++ b/Engine/Plugins/Experimental/ImageWidgets/Source/ImageWidgets/Private/SImageViewport.cpp
@@ -111,8 +111,7 @@ namespace UE::ImageWidgets
 			Commands.MipMinus,
 			FExecuteAction::CreateLambda([this]()
 			{
-				const TSharedPtr<FImageViewportClient> ViewportClient = StaticCastSharedPtr<FImageViewportClient>(Client);
-				const int32 CurrentMip = ViewportClient->GetMipLevel();
+				const int32 CurrentMip = ImageViewportClient->GetMipLevel();
 
 				int32 NumMips = 0;
 				if (ImageViewer.IsValid())
@@ -126,18 +125,17 @@ namespace UE::ImageWidgets
 
 				if (NumMips > 0 && CurrentMip < NumMips - 1)
 				{
-					ViewportClient->SetMipLevel(CurrentMip + 1);
+					ImageViewportClient->SetMipLevel(CurrentMip + 1);
 				}
 			}));
 		CommandList->MapAction(
 			Commands.MipPlus,
 			FExecuteAction::CreateLambda([this]()
 			{
-				const TSharedPtr<FImageViewportClient> ViewportClient = StaticCastSharedPtr<FImageViewportClient>(Client);
-				const int32 CurrentMip = ViewportClient->GetMipLevel();
+				const int32 CurrentMip = ImageViewportClient->GetMipLevel();
 				if (CurrentMip > INDEX_NONE)
 				{
-					ViewportClient->SetMipLevel(CurrentMip - 1);
+					ImageViewportClient->SetMipLevel(CurrentMip - 1);
 				}
 			}));
 	}
@@ -428,8 +426,8 @@ namespace UE::ImageWidgets
 				const int32 Mip = ImageViewportClient->GetMipLevel();
 				if (Mip > 0)
 				{
-					Size.X /= 1 << Mip;
-					Size.Y /= 1 << Mip;
+					Size.X = FMath::Max(1, Size.X / (1 << Mip));
+					Size.Y = FMath::Max(1, Size.Y / (1 << Mip));
 				}
 
 				return FText::Format(LOCTEXT("Resolution", "{0} \u00D7 {1}"), FText::AsNumber(Size.X), FText::AsNumber(Size.Y));
diff --git a/Engine/Plugins/Experimental/ImageWidgets/Source/ImageWidgets/Public/IImageViewer.h b/Engine/Plugins/Experimental/ImageWidgets/Source/ImageWidgets/Public/IImageViewer.h
--- a/Engine/Plugins/Experimental/ImageWidgets/Source/ImageWidgets/Public/IImageViewer.h
+++ b/Engine/Plugins/Experimental/ImageWidgets/Source/ImageWidgets/Public/IImageViewer.h
@@ -2,8 +2,8 @@
 
 #pragma once
 
-#include <CoreMinimal.h>
-#include <Misc/TVariant.h>
+#include "CoreMinimal.h"
+#include "Misc/TVariant.h"
 
 class FCanvas;
 class FViewport;
diff --git a/Engine/Plugins/Experimental/ImageWidgets/Source/ImageWidgets/Public/SImageCatalog.h b/Engine/Plugins/Experimental/ImageWidgets/Source/ImageWidgets/Public/SImageCatalog.h
--- a/Engine/Plugins/Experimental/ImageWidgets/Source/ImageWidgets/Public/SImageCatalog.h
+++ b/Engine/Plugins/Experimental/ImageWidgets/Source/ImageWidgets/Public/SImageCatalog.h
@@ -2,11 +2,9 @@
 
 #pragma once
 
-#include <CoreMinimal.h>
-#include <Framework/Views/ITypedTableView.h>
-#include <Widgets/SCompoundWidget.h>
-
-#define UE_API IMAGEWIDGETS_API
+#include "CoreMinimal.h"
+#include "Framework/Views/ITypedTableView.h"
+#include "Widgets/SCompoundWidget.h"
 
 namespace UE::ImageWidgets
 {
@@ -15,7 +13,7 @@ namespace UE::ImageWidgets
 	 */
 	struct FImageCatalogItemData
 	{
-		UE_API FImageCatalogItemData(FGuid Guid, const FSlateBrush& Brush, const FText& Name, const FText& Info, const FText& ToolTip);
+		IMAGEWIDGETS_API FImageCatalogItemData(FGuid Guid, const FSlateBrush& Brush, const FText& Name, const FText& Info, const FText& ToolTip);
 
 		/** Unique identifier for the catalog item */
 		FGuid Guid;
@@ -374,5 +372,3 @@ namespace UE::ImageWidgets
 		TPimplPtr<class FImpl> Impl;
 	};
 }
-
-#undef UE_API
diff --git a/Engine/Plugins/Experimental/ImageWidgets/Source/ImageWidgets/Public/SImageViewport.h b/Engine/Plugins/Experimental/ImageWidgets/Source/ImageWidgets/Public/SImageViewport.h
--- a/Engine/Plugins/Experimental/ImageWidgets/Source/ImageWidgets/Public/SImageViewport.h
+++ b/Engine/Plugins/Experimental/ImageWidgets/Source/ImageWidgets/Public/SImageViewport.h
@@ -2,8 +2,8 @@
 
 #pragma once
 
-#include <SEditorViewport.h>
-#include <Framework/MultiBox/MultiBoxExtender.h>
+#include "SEditorViewport.h"
+#include "Framework/MultiBox/MultiBoxExtender.h"
 
 class SViewportToolBar;
 
@@ -144,7 +144,7 @@ namespace UE::ImageWidgets
 			/** Do not show the bottom left status bar. */
 			bool bDisableStatusBarLeft = false;
 
-			/** Do not show the bottom left status bar. */
+			/** Do not show the bottom right status bar. */
 			bool bDisableStatusBarRight = false;
 		};
 
