ImageWidgets: Move out of Experimental

#Synced-CL 52049710

diff --git a/Engine/Plugins/Editor/ImageWidgets/ImageWidgets.uplugin b/Engine/Plugins/Editor/ImageWidgets/ImageWidgets.uplugin
new file mode 100644
--- /dev/null
+++ b/Engine/Plugins/Editor/ImageWidgets/ImageWidgets.uplugin
@@ -0,0 +1,16 @@
+﻿{
+  "FileVersion": 3,
+  "Version": 1,
+  "VersionName": "1.0",
+  "Description": "Generic Slate widgets for displaying images and image-like content, and content related to images.",
+  "Category": "Other",
+  "CreatedBy": "Epic Games, Inc.",
+  "CreatedByURL": "https://epicgames.com",
+  "IsBetaVersion": true,
+  "Modules": [
+    {
+      "Name": "ImageWidgets",
+      "Type": "Editor"
+    }
+  ]
+}
diff --git a/Engine/Plugins/Editor/ImageWidgets/Source/ImageWidgets/ImageWidgets.Build.cs b/Engine/Plugins/Editor/ImageWidgets/Source/ImageWidgets/ImageWidgets.Build.cs
new file mode 100644
--- /dev/null
+++ b/Engine/Plugins/Editor/ImageWidgets/Source/ImageWidgets/ImageWidgets.Build.cs
@@ -0,0 +1,44 @@
+﻿// Copyright Epic Games, Inc. All Rights Reserved.
+
+using UnrealBuildTool;
+
+public class ImageWidgets : ModuleRules
+{
+	// Set this flag to true to force building the ColorViewer sample. Otherwise, it will only be built in Debug builds. 
+	private const bool ForceBuildColorViewerSample = false;
+
+	public ImageWidgets(ReadOnlyTargetRules Target) : base(Target)
+	{
+		PCHUsage = PCHUsageMode.UseExplicitOrSharedPCHs;
+
+		PublicDependencyModuleNames.AddRange(
+			[
+				"Core",
+				"Slate",
+				"SlateCore",
+				"UnrealEd"
+			]
+		);
+
+		PrivateDependencyModuleNames.AddRange(
+			[
+				"ApplicationCore",
+				"CoreUObject",
+				"Engine",
+				"InputCore"
+			]
+		);
+
+		bool bBuildColorViewerSample = ForceBuildColorViewerSample || Target.Configuration == UnrealTargetConfiguration.Debug;
+		if (bBuildColorViewerSample)
+		{
+			PrivateDependencyModuleNames.AddRange(
+				[
+					"Projects",
+					"WorkspaceMenuStructure"
+				]
+			);
+		}
+		PublicDefinitions.Add(bBuildColorViewerSample ? "IMAGE_WIDGETS_BUILD_COLOR_VIEWER_SAMPLE=1" : "IMAGE_WIDGETS_BUILD_COLOR_VIEWER_SAMPLE=0");
+	}
+}
\ No newline at end of file
diff --git a/Engine/Plugins/Editor/ImageWidgets/Source/ImageWidgets/Private/ColorViewerSample/ColorViewer.cpp b/Engine/Plugins/Editor/ImageWidgets/Source/ImageWidgets/Private/ColorViewerSample/ColorViewer.cpp
new file mode 100644
--- /dev/null
+++ b/Engine/Plugins/Editor/ImageWidgets/Source/ImageWidgets/Private/ColorViewerSample/ColorViewer.cpp
@@ -0,0 +1,180 @@
+﻿// Copyright Epic Games, Inc. All Rights Reserved.
+
+#if IMAGE_WIDGETS_BUILD_COLOR_VIEWER_SAMPLE
+
+#include "ColorViewer.h"
+
+#include "CanvasItem.h"
+#include "CanvasTypes.h"
+
+#define LOCTEXT_NAMESPACE "ColorViewer"
+
+namespace UE::ImageWidgets::Sample
+{
+namespace Private
+{
+	static constexpr FGuid InvalidGuid{0, 0, 0, 0}; 
+
+	bool IsValid(const FGuid& Guid)
+	{
+		return Guid.A != 0;
+	}
+
+	void Invalidate(FGuid& Guid)
+	{
+		Guid.A = 0;
+	}
+	
+	int32 GetIndex(const FGuid& Guid)
+	{
+		return Guid.B;
+	}
+}
+
+	IImageViewer::FImageInfo FColorViewer::GetCurrentImageInfo() const
+	{
+		if (ColorIndexIsValid(SelectedColorIndex) && ColorGuidIsValid(Colors[SelectedColorIndex].Guid))
+		{
+			return {Colors[SelectedColorIndex].Guid, ImageSize, 0, true};
+		}
+		return {Private::InvalidGuid, FIntPoint::ZeroValue, 0, false};
+	}
+
+	void FColorViewer::DrawCurrentImage(FViewport* Viewport, FCanvas* Canvas, const FDrawProperties& Properties)
+	{
+		if (Properties.ABComparison.IsActive())
+		{
+			DrawImage(Properties.ABComparison.GuidA.B, Canvas, Properties.Placement, {0.0, 0.0}, {Properties.ABComparison.Threshold, 1.0});
+			DrawImage(Properties.ABComparison.GuidB.B, Canvas, Properties.Placement, {Properties.ABComparison.Threshold, 0.0}, {1.0, 1.0});
+		}
+		else
+		{
+			DrawImage(SelectedColorIndex, Canvas, Properties.Placement, {0.0, 0.0}, {1.0, 1.0});
+		}
+	}
+
+	TOptional<TVariant<FColor, FLinearColor>> FColorViewer::GetCurrentImagePixelColor(FIntPoint PixelCoords, int32 MipLevel) const
+	{
+		if (ColorIndexIsValid(SelectedColorIndex) && ColorGuidIsValid(Colors[SelectedColorIndex].Guid))
+		{
+			// Returns the current color.
+			// In a less trivial use case, the pixel coordinates and potentially the MIP level would be needed to look up the color value.
+			return TVariant<FColor, FLinearColor>(TInPlaceType<FColor>(), Colors[SelectedColorIndex].Color);
+		}
+		return {};
+	}
+
+	void FColorViewer::OnImageSelected(const FGuid& Guid)
+	{
+		if (ColorGuidIsValid(Guid))
+		{
+			SelectedColorIndex = Private::GetIndex(Guid);
+		}
+	}
+
+	bool FColorViewer::IsValidImage(const FGuid& Guid) const
+	{
+		return ColorGuidIsValid(Guid);
+	}
+
+	FText FColorViewer::GetImageName(const FGuid& Guid) const
+	{
+		if (ColorGuidIsValid(Guid))
+		{
+			const FColor& Color = Colors[Private::GetIndex(Guid)].Color;
+			const FString HexColor = FString::Printf(TEXT("#%02X%02X%02X"), Color.R, Color.G, Color.B);
+			return FText::FromString(HexColor);
+		}
+		return {};
+	}
+
+	const FColorViewer::FColorItem* FColorViewer::AddColor()
+	{
+		Colors.Add({FGuid(1, Colors.Num(), 0, 0), {}, FDateTime::Now()});
+
+		SelectedColorIndex = Colors.Num() - 1;
+
+		RandomizeColor();
+
+		return &Colors[SelectedColorIndex];
+	}
+
+	void FColorViewer::RemoveColor(const FGuid& Guid)
+	{
+		if (ColorGuidIsValid(Guid))
+		{
+			const int32 Index = Private::GetIndex(Guid);
+
+			Private::Invalidate(Colors[Index].Guid);
+
+			if (SelectedColorIndex == Index)
+			{
+				SelectedColorIndex = INDEX_NONE;
+			}
+		}
+	}
+
+	const FColorViewer::FColorItem* FColorViewer::RandomizeColor()
+	{
+		if (ColorIndexIsValid(SelectedColorIndex))
+		{
+			auto Random = []
+			{
+				return static_cast<uint8>(FMath::RandRange(0, 255));
+			};
+
+			Colors[SelectedColorIndex].Color = {Random(), Random(), Random()};
+
+			return &Colors[SelectedColorIndex];
+		}
+
+		return nullptr;
+	}
+
+	FToneMapping::EMode FColorViewer::GetToneMapping() const
+	{
+		return ToneMapping.Mode;
+	}
+
+	void FColorViewer::SetToneMapping(FToneMapping::EMode Mode)
+	{
+		ToneMapping.Mode = Mode;
+	}
+
+	FLinearColor FColorViewer::GetDefaultToneMappedColor(const FColor& Color) const
+	{
+		return ToneMapping.GetToneMappedColor(Color);
+	}
+
+	bool FColorViewer::ColorIndexIsValid(int32 Index) const
+	{
+		return 0 <= Index && Index < Colors.Num();
+	}
+
+	bool FColorViewer::ColorGuidIsValid(const FGuid& Guid) const
+	{
+		return Guid.A != 0 && ColorIndexIsValid(Guid.B) && Colors[Guid.B].Guid == Guid;
+	}
+
+	void FColorViewer::DrawImage(int32 Index, FCanvas* Canvas, const FDrawProperties::FPlacement& Placement, const FVector2d& UV0, const FVector2d& UV1) const
+	{
+		if (ColorIndexIsValid(Index))
+		{
+			// Get color value after tone mapping.
+			const FLinearColor ToneMappedColor = ToneMapping.GetToneMappedColor(Colors[Index].Color);
+
+			// Adjust offset and size based on which part of the image to draw.
+			const FVector2d Offset = Placement.Offset + Placement.Size * UV0;
+			const FVector2d Size = Placement.Size * (UV1 - UV0);
+
+			// Draw simple quad with current tone mapped color.
+			// In a less trivial use case, this would require rendering quads with textures and the like. 
+			FCanvasTileItem Tile(Offset, Size, ToneMappedColor);
+			Canvas->DrawItem(Tile);
+		}
+	}
+}
+
+#undef LOCTEXT_NAMESPACE
+
+#endif // IMAGE_WIDGETS_BUILD_COLOR_VIEWER_SAMPLE
diff --git a/Engine/Plugins/Editor/ImageWidgets/Source/ImageWidgets/Private/ColorViewerSample/ColorViewer.h b/Engine/Plugins/Editor/ImageWidgets/Source/ImageWidgets/Private/ColorViewerSample/ColorViewer.h
new file mode 100644
--- /dev/null
+++ b/Engine/Plugins/Editor/ImageWidgets/Source/ImageWidgets/Private/ColorViewerSample/ColorViewer.h
@@ -0,0 +1,109 @@
+﻿// Copyright Epic Games, Inc. All Rights Reserved.
+
+#pragma once
+
+#if IMAGE_WIDGETS_BUILD_COLOR_VIEWER_SAMPLE
+
+#include "IImageViewer.h"
+
+namespace UE::ImageWidgets::Sample
+{
+	/**
+	 * Provide tone mapping capabilities.
+	 * In this simple example, this is limited to just normal RGB plus luminance (grayscale).
+	 */
+	struct FToneMapping
+	{
+		enum class EMode { RGB, Lum };
+
+		FToneMapping(EMode Mode)
+			: Mode(Mode)
+		{
+		}
+
+		FLinearColor GetToneMappedColor(const FLinearColor& Color) const
+		{
+			if (Mode == EMode::Lum)
+			{
+				const float Luminance = 0.3f * Color.R + 0.59f * Color.G + 0.11f * Color.B;
+				return FLinearColor(Luminance, Luminance, Luminance);
+			}
+			return Color;
+		}
+
+		EMode Mode;
+	};
+
+	/**
+	 * Image viewer implementation used by the image widgets.
+	 * It contains any image data, in this case just colors, and renders the image data in the viewport widgets based on the viewport provided parameters.
+	 */
+	class FColorViewer final : public IImageViewer
+	{
+	public:
+		/**
+		 * Necessary data for a color item.
+		 */
+		struct FColorItem
+		{
+			/** Unique identifier for each item */
+			FGuid Guid;
+
+			/** The actual color value */
+			FColor Color;
+
+			/** Timestamp for when the item was created */
+			FDateTime DateTime;
+		};
+
+		// IImageViewer overrides - begin
+		virtual FImageInfo GetCurrentImageInfo() const override;
+		virtual void DrawCurrentImage(FViewport* Viewport, FCanvas* Canvas, const FDrawProperties& Properties) override;
+		virtual TOptional<TVariant<FColor, FLinearColor>> GetCurrentImagePixelColor(FIntPoint PixelCoords, int32 MipLevel) const override;
+		virtual void OnImageSelected(const FGuid& Guid) override;
+		virtual bool IsValidImage(const FGuid& Guid) const override;
+		virtual FText GetImageName(const FGuid& Guid) const override;
+		// IImageViewer overrides - end
+
+		/** Adds a color item. */
+		const FColorItem* AddColor();
+
+		/** Removes a color item. */
+		void RemoveColor(const FGuid& Guid);
+		
+		/** Sets a random color as the current "image" as a simple proxy for the image content changing and/or users choosing different images to display. */
+		const FColorItem* RandomizeColor();
+
+		/** Access to tone mapping data. This is effectively used by the viewport toolbar extensions as well as when drawing the image. */
+		FToneMapping::EMode GetToneMapping() const;
+		void SetToneMapping(FToneMapping::EMode Mode);
+
+		/** Apply default tone mapping to a given color. This is used to generate the catalog thumbnail. */
+		FLinearColor GetDefaultToneMappedColor(const FColor& Color) const;
+
+		/** Hardcoded values for the image size for all color.
+		 *  In a more realistic application, this value would depend on the actual current image. */
+		inline static const FIntPoint ImageSize = 512;
+
+	private:
+		/** Checks if a given index is a valid image. */
+		bool ColorIndexIsValid(int32 Index) const;
+
+		/** Checks if a given guid is a valid image. */
+		bool ColorGuidIsValid(const FGuid& Guid) const;
+
+		/** Draws the color image with the given index. The UVs determine if all or only a part of the image is drawn, i.e. for AB comparisons. */
+		void DrawImage(int32 Index, FCanvas* Canvas, const FDrawProperties::FPlacement& Placement, const FVector2d& UV0, const FVector2d& UV1) const;
+		
+		/** The tone mapping data. */
+		FToneMapping ToneMapping = FToneMapping::EMode::RGB;
+
+		/** The color for the currently displayed image. */
+		TArray<FColorItem> Colors;
+
+		/** Index of the currently selected image. */
+		int32 SelectedColorIndex = INDEX_NONE;
+	};
+}
+
+#endif // IMAGE_WIDGETS_BUILD_COLOR_VIEWER_SAMPLE
diff --git a/Engine/Plugins/Editor/ImageWidgets/Source/ImageWidgets/Private/ColorViewerSample/ColorViewerCommands.cpp b/Engine/Plugins/Editor/ImageWidgets/Source/ImageWidgets/Private/ColorViewerSample/ColorViewerCommands.cpp
new file mode 100644
--- /dev/null
+++ b/Engine/Plugins/Editor/ImageWidgets/Source/ImageWidgets/Private/ColorViewerSample/ColorViewerCommands.cpp
@@ -0,0 +1,31 @@
+﻿// Copyright Epic Games, Inc. All Rights Reserved.
+
+#if IMAGE_WIDGETS_BUILD_COLOR_VIEWER_SAMPLE
+
+#include "ColorViewerCommands.h"
+#include "ColorViewerStyle.h"
+
+#include "Framework/Commands/InputChord.h"
+
+#define LOCTEXT_NAMESPACE "ColorViewerCommands"
+
+namespace UE::ImageWidgets::Sample
+{
+	FColorViewerCommands::FColorViewerCommands()
+		: TCommands("ColorViewer", LOCTEXT("ContextDescription", "Color Viewer"), NAME_None, FColorViewerStyle::Get().GetStyleSetName())
+	{
+	}
+
+	void FColorViewerCommands::RegisterCommands()
+	{
+		UI_COMMAND(AddColor, "Add", "Add color to catalog", EUserInterfaceActionType::Button, FInputChord(EKeys::A));
+		UI_COMMAND(RandomizeColor, "Randomize", "Set random color", EUserInterfaceActionType::Button, FInputChord(EKeys::R));
+
+		UI_COMMAND(ToneMappingRGB, "RGB", "Show full color", EUserInterfaceActionType::ToggleButton, FInputChord());
+		UI_COMMAND(ToneMappingLum, "Lum", "Show luminance only", EUserInterfaceActionType::ToggleButton, FInputChord());
+	}
+}
+
+#undef LOCTEXT_NAMESPACE
+
+#endif // IMAGE_WIDGETS_BUILD_COLOR_VIEWER_SAMPLE
diff --git a/Engine/Plugins/Editor/ImageWidgets/Source/ImageWidgets/Private/ColorViewerSample/ColorViewerCommands.h b/Engine/Plugins/Editor/ImageWidgets/Source/ImageWidgets/Private/ColorViewerSample/ColorViewerCommands.h
new file mode 100644
--- /dev/null
+++ b/Engine/Plugins/Editor/ImageWidgets/Source/ImageWidgets/Private/ColorViewerSample/ColorViewerCommands.h
@@ -0,0 +1,29 @@
+﻿// Copyright Epic Games, Inc. All Rights Reserved.
+
+#pragma once
+
+#if IMAGE_WIDGETS_BUILD_COLOR_VIEWER_SAMPLE
+
+#include "Framework/Commands/Commands.h"
+#include "Framework/Commands/UICommandInfo.h"
+
+namespace UE::ImageWidgets::Sample
+{
+	/**
+	 * Provides commands used for the color viewer sample.
+	 */
+	class FColorViewerCommands : public TCommands<FColorViewerCommands>
+	{
+	public:
+		FColorViewerCommands();
+
+		virtual void RegisterCommands() override;
+
+		TSharedPtr<FUICommandInfo> AddColor; // Adding a new color entry
+		TSharedPtr<FUICommandInfo> RandomizeColor; // Setting current entry to a random color
+		TSharedPtr<FUICommandInfo> ToneMappingRGB; // Switching to RGB (no tone mapping)
+		TSharedPtr<FUICommandInfo> ToneMappingLum; // Switching to luminance (grayscale tone mapping)
+	};
+}
+
+#endif // IMAGE_WIDGETS_BUILD_COLOR_VIEWER_SAMPLE
diff --git a/Engine/Plugins/Editor/ImageWidgets/Source/ImageWidgets/Private/ColorViewerSample/ColorViewerStyle.cpp b/Engine/Plugins/Editor/ImageWidgets/Source/ImageWidgets/Private/ColorViewerSample/ColorViewerStyle.cpp
new file mode 100644
--- /dev/null
+++ b/Engine/Plugins/Editor/ImageWidgets/Source/ImageWidgets/Private/ColorViewerSample/ColorViewerStyle.cpp
@@ -0,0 +1,43 @@
+﻿// Copyright Epic Games, Inc. All Rights Reserved.
+
+#if IMAGE_WIDGETS_BUILD_COLOR_VIEWER_SAMPLE
+
+#include "ColorViewerStyle.h"
+
+#include "Interfaces/IPluginManager.h"
+#include "Styling/AppStyle.h"
+#include "Styling/CoreStyle.h"
+#include "Styling/SlateStyleRegistry.h"
+#include "Styling/SlateTypes.h"
+
+namespace UE::ImageWidgets::Sample
+{
+	FName FColorViewerStyle::StyleName("ColorViewerStyle");
+
+	FColorViewerStyle::FColorViewerStyle()
+		: FSlateStyleSet(StyleName)
+	{
+		const FString ContentDir = IPluginManager::Get().FindPlugin(TEXT("ImageWidgets"))->GetContentDir();
+
+		const FVector2f IconSizeToolbar(16.0f, 16.0f);
+
+		// Use custom icons in the content folder of this plugin.
+		Set("ToneMappingRGB", new FSlateImageBrush(ContentDir / TEXT("Icons/ToneMappingRGB") + TEXT(".png"), IconSizeToolbar));
+		Set("ToneMappingLum", new FSlateImageBrush(ContentDir / TEXT("Icons/ToneMappingLum") + TEXT(".png"), IconSizeToolbar));
+
+		FSlateStyleRegistry::RegisterSlateStyle(*this);
+	}
+
+	FColorViewerStyle::~FColorViewerStyle()
+	{
+		FSlateStyleRegistry::UnRegisterSlateStyle(*this);
+	}
+
+	FColorViewerStyle& FColorViewerStyle::Get()
+	{
+		static FColorViewerStyle Instance;
+		return Instance;
+	}
+}
+
+#endif // IMAGE_WIDGETS_BUILD_COLOR_VIEWER_SAMPLE
diff --git a/Engine/Plugins/Editor/ImageWidgets/Source/ImageWidgets/Private/ColorViewerSample/ColorViewerStyle.h b/Engine/Plugins/Editor/ImageWidgets/Source/ImageWidgets/Private/ColorViewerSample/ColorViewerStyle.h
new file mode 100644
--- /dev/null
+++ b/Engine/Plugins/Editor/ImageWidgets/Source/ImageWidgets/Private/ColorViewerSample/ColorViewerStyle.h
@@ -0,0 +1,27 @@
+﻿// Copyright Epic Games, Inc. All Rights Reserved.
+
+#pragma once
+
+#if IMAGE_WIDGETS_BUILD_COLOR_VIEWER_SAMPLE
+
+#include "Styling/SlateStyle.h"
+
+namespace UE::ImageWidgets::Sample
+{
+	/**
+	 * Style declarations for the color viewer sample.
+	 */
+	class FColorViewerStyle : public FSlateStyleSet
+	{
+	public:
+		static FName StyleName;
+
+		static FColorViewerStyle& Get();
+
+	private:
+		FColorViewerStyle();
+		virtual ~FColorViewerStyle() override;
+	};
+}
+
+#endif // IMAGE_WIDGETS_BUILD_COLOR_VIEWER_SAMPLE
\ No newline at end of file
diff --git a/Engine/Plugins/Editor/ImageWidgets/Source/ImageWidgets/Private/ColorViewerSample/ColorViewerWidget.cpp b/Engine/Plugins/Editor/ImageWidgets/Source/ImageWidgets/Private/ColorViewerSample/ColorViewerWidget.cpp
new file mode 100644
--- /dev/null
+++ b/Engine/Plugins/Editor/ImageWidgets/Source/ImageWidgets/Private/ColorViewerSample/ColorViewerWidget.cpp
@@ -0,0 +1,336 @@
+﻿// Copyright Epic Games, Inc. All Rights Reserved.
+
+#if IMAGE_WIDGETS_BUILD_COLOR_VIEWER_SAMPLE
+
+#include "ColorViewerWidget.h"
+
+#include "ColorViewerCommands.h"
+#include "ColorViewerStyle.h"
+#include "Algo/AllOf.h"
+#include "Algo/AnyOf.h"
+#include "Framework/MultiBox/MultiBoxBuilder.h"
+
+#define LOCTEXT_NAMESPACE "ColorViewerWidget"
+
+namespace UE::ImageWidgets::Sample
+{
+	void SColorViewerWidget::Construct(const FArguments&)
+	{
+		ColorViewer = MakeShared<FColorViewer>();
+
+		GroupColors = FName("Colors");
+		GroupFavorites = FName("Favorites");
+
+		BindCommands();
+
+		// Create toolbar extensions for a button to randomize the displayed color as well as the tone mapping controls.
+		const TSharedPtr<FExtender> ToolbarExtender = MakeShared<FExtender>();
+		ToolbarExtender->AddToolBarExtension("ToolbarCenter", EExtensionHook::Before, CommandList,
+		                                     FToolBarExtensionDelegate::CreateSP(this, &SColorViewerWidget::AddColorButtons));
+		ToolbarExtender->AddToolBarExtension("ToolbarRight", EExtensionHook::After, CommandList,
+		                                     FToolBarExtensionDelegate::CreateSP(this, &SColorViewerWidget::AddToneMappingButtons));
+
+		// Fill the widget with the image viewport.
+		ChildSlot
+			.VAlign(VAlign_Fill)
+			.HAlign(HAlign_Fill)
+			[
+				SAssignNew(Splitter, SSplitter)
+					.PhysicalSplitterHandleSize(2.0f)
+					+ SSplitter::Slot()
+						.Value(0.2f)
+						[
+							SAssignNew(Catalog, ImageWidgets::SImageCatalog)
+								.DefaultGroupName(GroupColors)
+								.DefaultGroupHeading(LOCTEXT("Colors", "Colors"))
+								.OnItemSelected_Lambda([&ColorViewer = ColorViewer, &Viewport = Viewport](const FGuid& ImageGuid)
+								{
+									if (ColorViewer)
+									{
+										ColorViewer->OnImageSelected(ImageGuid);
+									}
+									if (Viewport)
+									{
+										Viewport->RequestRedraw();
+									}
+								})
+								.OnGetGroupContextMenu(this, &SColorViewerWidget::GetGroupContextMenu)
+								.OnGetItemsContextMenu(this, &SColorViewerWidget::GetItemsContextMenu)
+						]
+					+ SSplitter::Slot()
+						.Value(0.8f)
+						[
+							SAssignNew(Viewport, ImageWidgets::SImageViewport, ColorViewer.ToSharedRef())
+								.ToolbarExtender(ToolbarExtender)
+								.DrawSettings(SImageViewport::FDrawSettings{
+									.ClearColor = FLinearColor::Black,
+									.bBorderEnabled = true,
+									.BorderThickness = 1.0,
+									.BorderColor = FVector3f(0.2f),
+									.bBackgroundColorEnabled = false,
+									.bBackgroundCheckerEnabled = false
+								})
+								.bABComparisonEnabled(true)
+						]
+			];
+
+		Catalog->AddGroup(GroupFavorites, LOCTEXT("Favorites", "Favorites"), GroupColors);
+	}
+
+	FReply SColorViewerWidget::OnKeyDown(const FGeometry& Geometry, const FKeyEvent& KeyEvent)
+	{
+		// Capture all key binds that are handled by the widget's commands. 
+		if (CommandList->ProcessCommandBindings(KeyEvent))
+		{
+			return FReply::Handled();
+		}
+
+		return FReply::Unhandled();
+	}
+
+	void SColorViewerWidget::AddColorButtons(FToolBarBuilder& ToolbarBuilder) const
+	{
+		const FSlateIcon AddColorIcon(FAppStyle::GetAppStyleSetName(), "FontEditor.Button_Add");
+		ToolbarBuilder.AddToolBarButton(FColorViewerCommands::Get().AddColor, NAME_None, TAttribute<FText>(), TAttribute<FText>(),
+										TAttribute<FSlateIcon>(AddColorIcon));
+
+		const FSlateIcon RandomizeColorIcon(FAppStyle::GetAppStyleSetName(), "FontEditor.Update");
+		ToolbarBuilder.AddToolBarButton(FColorViewerCommands::Get().RandomizeColor, NAME_None, TAttribute<FText>(), TAttribute<FText>(),
+		                                TAttribute<FSlateIcon>(RandomizeColorIcon));
+	}
+
+	void SColorViewerWidget::AddToneMappingButtons(FToolBarBuilder& ToolbarBuilder) const
+	{
+		ToolbarBuilder.AddSeparator();
+
+		ToolbarBuilder.BeginBlockGroup();
+		{
+			const FName& StyleSetName = FColorViewerStyle::Get().GetStyleSetName();
+			const FColorViewerCommands& Commands = FColorViewerCommands::Get();
+
+			const FSlateIcon RGBIcon(StyleSetName, "ToneMappingRGB");
+			const FSlateIcon LumIcon(StyleSetName, "ToneMappingLum");
+
+			ToolbarBuilder.AddToolBarButton(Commands.ToneMappingRGB, NAME_None, TAttribute<FText>(), TAttribute<FText>(), TAttribute<FSlateIcon>(RGBIcon));
+			ToolbarBuilder.AddToolBarButton(Commands.ToneMappingLum, NAME_None, TAttribute<FText>(), TAttribute<FText>(), TAttribute<FSlateIcon>(LumIcon));
+		}
+		ToolbarBuilder.EndBlockGroup();
+	}
+
+	TTuple<FText, FText, FText> GetColorItemMetaData(const FColorViewer& ColorViewer, const FColorViewer::FColorItem* ColorItem)
+	{
+		const FText Name = ColorViewer.GetImageName(ColorItem->Guid);
+
+		const FText Info = FText::Format(
+			LOCTEXT("ColorEntryInfoLabel", "{0}"), FText::AsTime(ColorItem->DateTime, EDateTimeStyle::Short, FText::GetInvariantTimeZone()));
+
+		const FText ToolTip = FText::Format(
+			LOCTEXT("ColorEntryToolTip", "R {0}, G {1}, B {2}"), {ColorItem->Color.R, ColorItem->Color.G, ColorItem->Color.B});
+
+		return {Name, Info, ToolTip};
+	}
+
+	void SColorViewerWidget::AddColor()
+	{
+		if (const FColorViewer::FColorItem* ColorItem = ColorViewer->AddColor())
+		{
+			auto [Name, Info, ToolTip] = GetColorItemMetaData(*ColorViewer, ColorItem);
+			const FLinearColor ToneMappedColor = ColorViewer->GetDefaultToneMappedColor(ColorItem->Color);
+
+			Catalog->AddItem(MakeShared<FImageCatalogItemData>(ColorItem->Guid, FSlateColorBrush(ToneMappedColor), Name, Info, ToolTip));
+			Catalog->ClearSelection();
+			Catalog->SelectItem(ColorItem->Guid);
+		}
+	}
+
+	void SColorViewerWidget::RandomizeColor()
+	{
+		if (const FColorViewer::FColorItem* ColorItem = ColorViewer->RandomizeColor())
+		{
+			const auto [Name, Info, ToolTip] = GetColorItemMetaData(*ColorViewer, ColorItem);
+			const FLinearColor ToneMappedColor = ColorViewer->GetDefaultToneMappedColor(ColorItem->Color);
+
+			Catalog->UpdateItem({ColorItem->Guid, FSlateColorBrush(ToneMappedColor), Name, Info, ToolTip});
+		}
+	}
+
+	bool SColorViewerWidget::RandomizeColorEnabled() const
+	{
+		const IImageViewer::FImageInfo Info = ColorViewer->GetCurrentImageInfo();
+		if (Info.bIsValid)
+		{
+			if (const TOptional<FName> GroupName = Catalog->GetItemGroupName(Info.Guid))
+			{
+				if (*GroupName == GroupColors)
+				{
+					return true;
+				}
+			}
+		}
+
+		return false;
+	}
+
+	TSharedPtr<SWidget> SColorViewerWidget::GetGroupContextMenu(FName GroupName) const
+	{
+		FMenuBuilder MenuBuilder(true, nullptr);
+
+		const int32 NumItems = Catalog->NumItems(GroupName);
+		if (NumItems > 0)
+		{
+			MenuBuilder.AddMenuEntry(
+				NumItems == 1
+					? FText::Format(LOCTEXT("DeleteGroupItem", "Delete single item in group"), NumItems)
+					: FText::Format(LOCTEXT("DeleteAllGroupItems", "Delete all {0} items in group"), NumItems),
+				LOCTEXT("DeleteGroupItems", "Deletes all items in this group."),
+				FSlateIcon(FAppStyle::GetAppStyleSetName(), "GenericCommands.Delete"),
+				FUIAction(FExecuteAction::CreateLambda(
+					[GroupName, Catalog = Catalog, ColorViewer = ColorViewer]
+					{
+						if (Catalog && ColorViewer)
+						{
+							for (int32 Index = Catalog->NumItems(GroupName) - 1; Index >= 0; --Index)
+							{
+								const TOptional<FGuid> Guid = Catalog->GetItemGuidAt(Index, GroupName);
+								if (Guid.IsSet())
+								{
+									Catalog->RemoveItem(Guid.GetValue());
+									ColorViewer->RemoveColor(Guid.GetValue());
+								}
+							}
+						}
+					})));
+		}
+
+		return MenuBuilder.MakeWidget();
+	}
+
+	TSharedPtr<SWidget> SColorViewerWidget::GetItemsContextMenu(const TArray<FGuid>& Guids) const
+	{
+		checkSlow(!Guids.IsEmpty());
+
+		auto IsFavorite = [&Catalog = Catalog, &GroupFavorites = GroupFavorites](const FGuid& Guid)
+		{
+			const TOptional<FName> GroupName = Catalog->GetItemGroupName(Guid);
+			return GroupName && *GroupName == GroupFavorites;
+		};
+
+		const bool bHaveFavorites = Algo::AnyOf(Guids, IsFavorite);
+		const bool bHaveNonFavorites = !Algo::AllOf(Guids, IsFavorite);
+
+		FMenuBuilder MenuBuilder(true, nullptr);
+
+		if (bHaveNonFavorites)
+		{
+			MenuBuilder.AddMenuEntry(
+				LOCTEXT("AddFavorite", "Add To Favorites"),
+				LOCTEXT("AddFavoriteTooltip", "Adds the selected color(s) to the list of favorites."),
+				FSlateIcon(FAppStyle::GetAppStyleSetName(), "Icons.Pinned"),
+				FUIAction(
+					FExecuteAction::CreateLambda(
+						[Guids, Catalog = Catalog, GroupFavorites = GroupFavorites, GroupColors = GroupColors]
+						{
+							if (Catalog)
+							{
+								for (const FGuid& Guid : Guids)
+								{
+									const auto GroupName = Catalog->GetItemGroupName(Guid);
+									if (GroupName && *GroupName == GroupColors)
+									{
+										Catalog->MoveItem(Guid, GroupFavorites);
+									}
+									Catalog->SelectItem(Guid);
+								}
+							}
+						})));
+		}
+
+		if (bHaveFavorites)
+		{
+			MenuBuilder.AddMenuEntry(
+				LOCTEXT("RemoveFavorite", "Remove from Favorites"),
+				LOCTEXT("RemoveFavoriteTooltip", "Removes the selected color(s) from the list of favorites."),
+				FSlateIcon(FAppStyle::GetAppStyleSetName(), "Icons.Unpinned"),
+				FUIAction(
+					FExecuteAction::CreateLambda(
+						[Guids, Catalog = Catalog, GroupFavorites = GroupFavorites, GroupColors = GroupColors]
+						{
+							if (Catalog)
+							{
+								for (const FGuid& Guid : Guids)
+								{
+									const auto GroupName = Catalog->GetItemGroupName(Guid);
+									if (GroupName && *GroupName == GroupFavorites)
+									{
+										Catalog->MoveItem(Guid, GroupColors);
+									}
+									Catalog->SelectItem(Guid);
+								}
+							}
+						})));
+		}
+
+		MenuBuilder.AddSeparator();
+
+		MenuBuilder.AddMenuEntry(
+			LOCTEXT("DeleteColors", "Delete"),
+			LOCTEXT("DeleteColorsTooltip", "Deletes the selected color(s) from the catalog."),
+			FSlateIcon(FAppStyle::GetAppStyleSetName(), "GenericCommands.Delete"),
+			FUIAction(FExecuteAction::CreateLambda(
+				[Guids, Catalog = Catalog, ColorViewer = ColorViewer]
+				{
+					if (Catalog && ColorViewer)
+					{
+						for (const FGuid& Guid : Guids)
+						{
+							Catalog->RemoveItem(Guid);
+							ColorViewer->RemoveColor(Guid);
+						}
+					}
+				})));
+
+		return MenuBuilder.MakeWidget();
+	}
+
+	void SColorViewerWidget::BindCommands()
+	{
+		const FColorViewerCommands& Commands = FColorViewerCommands::Get();
+
+		CommandList = MakeShared<FUICommandList>();
+
+		CommandList->MapAction(
+			Commands.AddColor,
+			FExecuteAction::CreateSP(this, &SColorViewerWidget::AddColor)
+		);
+
+		CommandList->MapAction(
+			Commands.RandomizeColor,
+			FExecuteAction::CreateSP(this, &SColorViewerWidget::RandomizeColor),
+			FCanExecuteAction::CreateSP(this, &SColorViewerWidget::RandomizeColorEnabled)
+		);
+
+		CommandList->MapAction(
+			Commands.ToneMappingRGB,
+			FExecuteAction::CreateSP(ColorViewer.ToSharedRef(), &FColorViewer::SetToneMapping, FToneMapping::EMode::RGB),
+			FCanExecuteAction(),
+			FIsActionChecked::CreateLambda([ColorViewer = ColorViewer]
+			{
+				return ColorViewer->GetToneMapping() == FToneMapping::EMode::RGB;
+			})
+		);
+
+		CommandList->MapAction(
+			Commands.ToneMappingLum,
+			FExecuteAction::CreateSP(ColorViewer.ToSharedRef(), &FColorViewer::SetToneMapping, FToneMapping::EMode::Lum),
+			FCanExecuteAction(),
+			FIsActionChecked::CreateLambda([ColorViewer = ColorViewer]
+			{
+				return ColorViewer->GetToneMapping() == FToneMapping::EMode::Lum;
+			})
+		);
+	}
+}
+
+#undef LOCTEXT_NAMESPACE
+
+#endif // IMAGE_WIDGETS_BUILD_COLOR_VIEWER_SAMPLE
diff --git a/Engine/Plugins/Editor/ImageWidgets/Source/ImageWidgets/Private/ColorViewerSample/ColorViewerWidget.h b/Engine/Plugins/Editor/ImageWidgets/Source/ImageWidgets/Private/ColorViewerSample/ColorViewerWidget.h
new file mode 100644
--- /dev/null
+++ b/Engine/Plugins/Editor/ImageWidgets/Source/ImageWidgets/Private/ColorViewerSample/ColorViewerWidget.h
@@ -0,0 +1,77 @@
+﻿// Copyright Epic Games, Inc. All Rights Reserved.
+
+#pragma once
+
+#if IMAGE_WIDGETS_BUILD_COLOR_VIEWER_SAMPLE
+
+#include "ColorViewer.h"
+
+#include "SImageCatalog.h"
+#include "SImageViewport.h"
+#include "Widgets/Layout/SSplitter.h"
+
+namespace UE::ImageWidgets::Sample
+{
+	/**
+	 * Widget that contains and configures the image widgets.
+	 */
+	class SColorViewerWidget : public SCompoundWidget
+	{
+	public:
+		SLATE_BEGIN_ARGS(SColorViewerWidget)
+			{
+			}
+
+		SLATE_END_ARGS()
+
+		/** Gets called by Slate for construction of this widget. */
+		void Construct(const FArguments& Args);
+
+		/** Forwards key presses to the command bindings. */
+		virtual FReply OnKeyDown(const FGeometry& Geometry, const FKeyEvent& KeyEvent) override;
+
+	private:
+		/** Binds all the commands. */
+		void BindCommands();
+
+		/** Functions for adding buttons as viewport toolbar extensions. */
+		void AddColorButtons(FToolBarBuilder& ToolbarBuilder) const;
+		void AddToneMappingButtons(FToolBarBuilder& ToolbarBuilder) const;
+
+		/** Add a new color entry to the catalog. */
+		void AddColor();
+
+		/** Choose a random color for the current entry. */
+		void RandomizeColor();
+
+		/** Checks if the Randomize Color action is available. */
+		bool RandomizeColorEnabled() const;
+
+		/** Creates a context menu for a group. */
+		TSharedPtr<SWidget> GetGroupContextMenu(FName GroupName) const;
+
+		/** Creates a context menu for a set of selected guids. */
+		TSharedPtr<SWidget> GetItemsContextMenu(const TArray<FGuid>& Guids) const;
+
+		/** The image viewer implementation that contains the image data and renders the image. */
+		TSharedPtr<FColorViewer> ColorViewer;
+
+		/** Adjustable divider between catalog on the left and viewport on the right. */
+		TSharedPtr<SSplitter> Splitter;
+
+		/** The image catalog that holds all currently available images. */
+		TSharedPtr<SImageCatalog> Catalog;
+
+		/** The image viewport in which the current image gets displayed in. */
+		TSharedPtr<SImageViewport> Viewport;
+
+		/** The commands used by this sample widget. */
+		TSharedPtr<FUICommandList> CommandList;
+
+		/** Names for groups in the catalog. */
+		FName GroupColors;
+		FName GroupFavorites;
+	};
+}
+
+#endif // IMAGE_WIDGETS_BUILD_COLOR_VIEWER_SAMPLE
diff --git a/Engine/Plugins/Editor/ImageWidgets/Source/ImageWidgets/Private/ImageABComparison.h b/Engine/Plugins/Editor/ImageWidgets/Source/ImageWidgets/Private/ImageABComparison.h
new file mode 100644
--- /dev/null
+++ b/Engine/Plugins/Editor/ImageWidgets/Source/ImageWidgets/Private/ImageABComparison.h
@@ -0,0 +1,83 @@
+﻿// Copyright Epic Games, Inc. All Rights Reserved.
+
+#pragma once
+
+#include "Delegates/Delegate.h"
+#include "Internationalization/Text.h"
+#include "Misc/Guid.h"
+
+namespace UE::ImageWidgets
+{
+	/**
+	 * Provides data and logic for AB comparisons of images.
+	 */
+	class FImageABComparison
+	{
+	public:
+		DECLARE_DELEGATE_RetVal(FGuid, FGetCurrentImageGuid)
+		DECLARE_DELEGATE_RetVal_OneParam(bool, FImageIsValid, const FGuid&)
+		DECLARE_DELEGATE_RetVal_OneParam(FText, FGetImageName, const FGuid&)
+
+		FImageABComparison(FImageIsValid&& InImageIsValid, FGetCurrentImageGuid&& InGetCurrentImageGuid, FGetImageName&& InGetImageName)
+			: ImageIsValid(MoveTemp(InImageIsValid))
+			, GetCurrentImageGuid(MoveTemp(InGetCurrentImageGuid))
+			, GetImageName(MoveTemp(InGetImageName))
+		{
+			check(ImageIsValid.IsBound());
+			check(GetCurrentImageGuid.IsBound());
+			check(GetImageName.IsBound());
+		}
+
+		enum class EAorB : int8
+		{
+			A,
+			B
+		};
+
+		bool CanSetABComparison(EAorB AorB) const
+		{
+			const FGuid CurrentGuid = GetCurrentImageGuid.Execute();
+			const bool bGuidIsValidAndCurrentGuid = Guids[static_cast<int32>(AorB)].IsValid() && Guids[static_cast<int32>(AorB)] == CurrentGuid;
+			const bool bViewerHasValidImage = ImageIsValid.Execute(CurrentGuid);
+			const bool bGuidIsDisabledAndOtherGuidIsNotCurrentGuid = !Guids[static_cast<int32>(AorB)].IsValid() && Guids[!static_cast<bool>(AorB)] != CurrentGuid;
+
+			return bGuidIsValidAndCurrentGuid || (bViewerHasValidImage && bGuidIsDisabledAndOtherGuidIsNotCurrentGuid);
+		}
+
+		void SetABComparison(EAorB AorB, const FGuid& Guid)
+		{
+			Guids[static_cast<int32>(AorB)] = Guid;
+		}
+
+		bool ABComparisonIsSet(EAorB AorB) const
+		{
+			return Guids[static_cast<int32>(AorB)].IsValid();
+		}
+
+		bool IsActive() const
+		{
+			return Guids[static_cast<int32>(EAorB::A)].IsValid() && Guids[static_cast<int32>(EAorB::B)].IsValid();
+		}
+
+		const FGuid& GetGuidA() const
+		{
+			return Guids[static_cast<int32>(EAorB::A)];
+		}
+
+		const FGuid& GetGuidB() const
+		{
+			return Guids[static_cast<int32>(EAorB::B)];
+		}
+
+		FText GetName(EAorB AorB) const
+		{
+			return Guids[static_cast<int32>(AorB)].IsValid() ? GetImageName.Execute(Guids[static_cast<int32>(AorB)]) : FText();
+		}
+
+	private:
+		FGuid Guids[2];
+		FImageIsValid ImageIsValid;
+		FGetCurrentImageGuid GetCurrentImageGuid;
+		FGetImageName GetImageName;
+	};
+}
diff --git a/Engine/Plugins/Editor/ImageWidgets/Source/ImageWidgets/Private/ImageViewportClient.cpp b/Engine/Plugins/Editor/ImageWidgets/Source/ImageWidgets/Private/ImageViewportClient.cpp
new file mode 100644
--- /dev/null
+++ b/Engine/Plugins/Editor/ImageWidgets/Source/ImageWidgets/Private/ImageViewportClient.cpp
@@ -0,0 +1,529 @@
+﻿// Copyright Epic Games, Inc. All Rights Reserved.
+
+#include "ImageViewportClient.h"
+
+#include "CanvasItem.h"
+#include "CanvasTypes.h"
+#include "ImageABComparison.h"
+#include "ImageWidgetsLogCategory.h"
+#include "SImageViewport.h"
+#include "Texture2DPreview.h"
+#include "TextureResource.h"
+#include "Engine/Texture2D.h"
+
+#define LOCTEXT_NAMESPACE "ImageViewportClient"
+
+namespace UE::ImageWidgets
+{
+
+UTexture2D* CreateCheckerTexture(const FColor& Color1, const FColor& Color2, const int32 CheckerSize)
+{
+	constexpr EPixelFormat Format = PF_B8G8R8A8;
+	const int32 TextureSize = FMath::Min<int32>(FMath::RoundUpToPowerOfTwo(FMath::Max(GPixelFormats[PF_B8G8R8A8].BlockSizeX, CheckerSize)) * 2, 4096);
+
+	UTexture2D* CheckerTexture = UTexture2D::CreateTransient(TextureSize, TextureSize, Format);
+	check(CheckerTexture != nullptr);
+
+	UE_LOG(LogImageWidgets, Verbose, TEXT("Creating background checker texture '%s' with Color1=%s, Color2=%s, and Size=%d."),
+	       *CheckerTexture->GetName(), *Color1.ToHex(), *Color2.ToHex(), CheckerSize);
+
+	// TODO Revisit setting the LODGroup, and figure out if there is a better way to do this.
+	// Setting the LOD group seems to be the only way to avoid interpolation when zoomed in. 
+	CheckerTexture->LODGroup = TEXTUREGROUP_Pixels2D;
+
+	TIndirectArray<FTexture2DMipMap>& Mips = CheckerTexture->GetPlatformData()->Mips;
+	int32 MipSize = TextureSize;
+	while (MipSize >= GPixelFormats[Format].BlockSizeX)
+	{
+		const int32 MipIndex = MipSize == TextureSize ? 0 : Mips.Add(new FTexture2DMipMap(MipSize, MipSize, 1));
+
+		FByteBulkData& MipBulkData = Mips[MipIndex].BulkData;
+
+		FColor *const MipColorData = [&MipBulkData, MipIndex, MipSize]
+		{
+			FColor* ColorData = static_cast<FColor*>(MipBulkData.Lock(LOCK_READ_WRITE));
+			if (MipIndex > 0)
+			{
+				const int64 NumBlocks = MipSize / GPixelFormats[Format].BlockSizeX;
+				const int64 NumBytes = NumBlocks * NumBlocks * GPixelFormats[Format].BlockBytes;
+
+				ColorData = reinterpret_cast<FColor*>(MipBulkData.Realloc(NumBytes));
+			}
+			return ColorData;
+		}();
+
+		if (MipSize > 1)
+		{
+			const int32 MipCheckerSize = MipSize / 2;
+			for (int32 Y = 0; Y < MipSize; ++Y)
+			{
+				const bool bTop = Y < MipCheckerSize;
+				for (int32 X = 0; X < MipSize; ++X)
+				{
+					const bool bLeft = X < MipCheckerSize;
+					MipColorData[X + Y * MipSize] = bTop == bLeft ? Color1 : Color2;
+				}
+			}
+		}
+		else
+		{
+			*MipColorData = FColor((static_cast<uint16>(Color1.R) + Color2.R) / 2,
+			                       (static_cast<uint16>(Color1.G) + Color2.G) / 2,
+			                       (static_cast<uint16>(Color1.B) + Color2.B) / 2);
+		}
+
+		MipBulkData.Unlock();
+
+		MipSize /= 2;
+	}
+
+	CheckerTexture->UpdateResource();
+
+	return CheckerTexture;
+}
+
+void DestroyCheckerTexture(TStrongObjectPtr<UTexture2D>& CheckerTexture)
+{
+	if (CheckerTexture)
+	{
+		UE_LOG(LogImageWidgets, Verbose, TEXT("Destroying background checker texture '%s'."), *CheckerTexture->GetName());
+
+		if (CheckerTexture->GetResource())
+		{
+			CheckerTexture->ReleaseResource();
+		}
+		CheckerTexture->MarkAsGarbage();
+
+		CheckerTexture = nullptr;
+	}
+}
+
+FImageViewportClient::FImageViewportClient(const TWeakPtr<SEditorViewport>& InEditorViewport, FGetImageSize&& InGetImageSize, FDrawImage&& InDrawImage,
+                                           FGetDrawSettings&& InGetDrawSettings, FGetLayoutAndDPIScaleFactor&& InGetLayoutAndDPIScaleFactor,
+                                           FGetLayoutScaleFactor&& InGetLayoutScaleFactor, const FImageABComparison* InABComparison,
+                                           const SImageViewport::FControllerSettings& InControllerSettings)
+	: FEditorViewportClient(nullptr, nullptr, InEditorViewport)
+	, GetImageSize(MoveTemp(InGetImageSize))
+	, DrawImage(MoveTemp(InDrawImage))
+	, GetDrawSettings(MoveTemp(InGetDrawSettings))
+	, GetLayoutAndDPIScaleFactor(MoveTemp(InGetLayoutAndDPIScaleFactor))
+	, GetLayoutScaleFactor(MoveTemp(InGetLayoutScaleFactor))
+	, OnInputKey(InControllerSettings.OnInputKey)
+	, ABComparison(InABComparison)
+	, Controller(static_cast<FImageViewportController::EZoomMode>(InControllerSettings.DefaultZoomMode))
+	, bZoomOnResize(InControllerSettings.bZoomOnResize)
+{
+	check(GetImageSize.IsBound());
+	check(DrawImage.IsBound());
+	check(GetDrawSettings.IsBound());
+	check(GetLayoutAndDPIScaleFactor.IsBound());
+	check(GetLayoutScaleFactor.IsBound());
+	check(ABComparison);
+
+	SetRealtime(true);
+}
+
+FImageViewportClient::~FImageViewportClient()
+{
+	if (CheckerTexture)
+	{
+		DestroyCheckerTexture(CheckerTexture);
+	}
+}
+
+void FImageViewportClient::Draw(FViewport* InViewport, FCanvas* Canvas)
+{
+	const SImageViewport::FDrawSettings DrawSettings = GetDrawSettings.Execute();
+	CreateOrDestroyCheckerTextureIfSettingsChanged(DrawSettings);
+
+	Canvas->Clear(DrawSettings.ClearColor);
+
+	const FIntPoint ImageSize = GetImageSize.Execute();
+	if (ImageSize.X <= 0 || ImageSize.Y <= 0)
+	{
+		return;
+	}
+
+	const float LayoutAndDPIScale = GetLayoutAndDPIScaleFactor.Execute();
+
+	// Recalculate zoom and pan when the viewport is resized.
+	{
+		const FVector2d ScaledSize = FVector2d(InViewport->GetSizeXY()) / LayoutAndDPIScale;
+		if (ScaledSize != CachedScaledViewportSize)
+		{
+			if (CachedScaledViewportSize != FVector2d::ZeroVector)
+			{
+				const FImageViewportController::FZoomSettings OldZoomSettings = Controller.GetZoom();
+				const FVector2d OldPan = Controller.GetPan(FVector2d::Zero());
+
+				if (OldZoomSettings.Mode != FImageViewportController::EZoomMode::Custom)
+				{
+					// Fill/Fit: recalculate zoom for new viewport size.
+					Controller.SetZoom(OldZoomSettings.Mode, OldZoomSettings.Zoom, ImageSize, ScaledSize);
+				}
+				else if (bZoomOnResize)
+				{
+					// When zoom on resize is set, scale zoom proportionally.
+					const double OldFillZoom = (CachedScaledViewportSize / FVector2d(ImageSize)).GetMin();
+					const double NewFillZoom = (ScaledSize / FVector2d(ImageSize)).GetMin();
+					if (OldFillZoom > SMALL_NUMBER)
+					{
+						Controller.SetZoom(FImageViewportController::EZoomMode::Custom, OldZoomSettings.Zoom * (NewFillZoom / OldFillZoom),
+						                   ImageSize, ScaledSize);
+					}
+				}
+
+				// Scale pan uniformly by the zoom ratio to maintain relative image position.
+				// Using a single scalar (NewZoom/OldZoom) for both axes avoids diagonal drift when the constraining axis flips, e.g. landscape to portrait.
+				const double NewZoom = Controller.GetZoom().Zoom;
+				const double PanScale = OldZoomSettings.Zoom > SMALL_NUMBER ? NewZoom / OldZoomSettings.Zoom : 1.0;
+				const FVector2d NewPan = OldPan * PanScale;
+				const FVector2d CurrentPan = Controller.GetPan(FVector2d::Zero());
+				Controller.Pan(NewPan - CurrentPan);
+			}
+
+			CachedScaledViewportSize = ScaledSize;
+		}
+	}
+
+	const float LayoutScale = GetLayoutScaleFactor.Execute();
+
+	// Controller computes placement in unzoomed logical space
+	const FVector2d ScaledViewportSize = FVector2d(InViewport->GetSizeXY()) / LayoutAndDPIScale;
+	CachedPlacement = GetPlacementProperties(ImageSize, ScaledViewportSize);
+	bCachedPlacementIsValid = true;
+
+	const IImageViewer::FDrawProperties::FMip MipProperties = GetMipProperties();
+
+	// Scale placement from unzoomed logical space to RT pixel space for drawing
+	const FVector2d DrawOffset = CachedPlacement.Offset * LayoutScale;
+	const FVector2d DrawSize = CachedPlacement.Size * LayoutScale;
+
+	// Draw border
+	if (DrawSettings.bBorderEnabled)
+	{
+		FCanvasBoxItem Border(DrawOffset - DrawSettings.BorderThickness / 2.0, DrawSize + DrawSettings.BorderThickness);
+		Border.LineThickness = DrawSettings.BorderThickness;
+		Border.SetColor(DrawSettings.BorderColor);
+		Canvas->DrawItem(Border);
+	}
+
+	// Draw background
+	if (DrawSettings.bBackgroundColorEnabled || DrawSettings.bBackgroundCheckerEnabled)
+	{
+		FCanvasTileItem Background(DrawOffset, DrawSize, DrawSettings.BackgroundColor);
+
+		TRefCountPtr<FBatchedElementTexture2DPreviewParameters> BatchedElementParameters;
+
+		if (DrawSettings.bBackgroundCheckerEnabled)
+		{
+			if (ensure(CheckerTexture != nullptr))
+			{
+				Background.SetColor(FLinearColor::White);
+				Background.Texture = CheckerTexture->GetResource();
+
+				Background.UV1 = FVector2D(ImageSize) / CheckerTexture->GetSizeX();
+
+				BatchedElementParameters = new FBatchedElementTexture2DPreviewParameters(
+					-1, 0, 0, false, false, false, false, false, false);
+				Background.BatchedElementParameters = BatchedElementParameters.GetReference();
+			}
+		}
+
+		Canvas->DrawItem(Background);
+	}
+
+	const double ABComparisonDividerX = [this, LayoutScale, &ImageSize]
+	{
+		if (bDraggingABComparisonDivider)
+		{
+			FIntPoint MousePos;
+			Viewport->GetMousePos(MousePos);
+
+			// Mouse pos is in RT pixel space; convert to unzoomed placement space
+			FVector2d DividerPos = FVector2d(MousePos) / LayoutScale - CachedPlacement.Offset;
+			DividerPos /= CachedPlacement.ZoomFactor;
+
+			return DividerPos.X / ImageSize.X;
+		}
+		return ABComparisonDivider;
+	}();
+
+	const IImageViewer::FDrawProperties::FABComparison ABComparisonProperties = [this, &ABComparisonDividerX]() -> IImageViewer::FDrawProperties::FABComparison
+	{
+		return {ABComparison->GetGuidA(), ABComparison->GetGuidB(), FMath::Clamp(ABComparisonDividerX, 0.0, 1.0)};
+	}();
+
+	// Draw image — pass scaled placement for RT pixel space drawing
+	const IImageViewer::FDrawProperties::FPlacement DrawPlacement{DrawOffset, DrawSize, CachedPlacement.ZoomFactor * LayoutScale};
+	DrawImage.Execute(InViewport, Canvas, {DrawPlacement, MipProperties, ABComparisonProperties});
+
+	if (ABComparison->IsActive())
+	{
+		const FVector2D LineStart(DrawOffset.X + DrawSize.X * ABComparisonDividerX, 0);
+		const FVector2D LineEnd(DrawOffset.X + DrawSize.X * ABComparisonDividerX, InViewport->GetSizeXY().Y);
+		FCanvasLineItem Line(LineStart, LineEnd);
+		Line.LineThickness = 2.0f;
+		Line.SetColor({0.5f, 0.5f, 0.5f, 1.0f});
+		Canvas->DrawItem(Line);
+	}
+}
+
+EMouseCursor::Type FImageViewportClient::GetCursor(FViewport* InViewport, int32 X, int32 Y)	
+{
+	if (bDragging)
+	{
+		CachedMouseX = X;
+		CachedMouseY = Y;
+		return EMouseCursor::GrabHandClosed;
+	}
+
+	if (bDraggingABComparisonDivider || MouseIsOverABComparisonDivider({X, Y}))
+	{
+		return EMouseCursor::ResizeLeftRight;
+	}
+
+	const auto [PixelCoordsValid, PixelCoords] = GetPixelCoordinatesUnderCursor();
+	const FIntPoint ImageSize = GetImageSize.Execute();
+	if (PixelCoordsValid && 0 <= PixelCoords.X && 0 <= PixelCoords.Y && PixelCoords.X < ImageSize.X && PixelCoords.Y < ImageSize.Y)
+	{
+		return EMouseCursor::Crosshairs;
+	}
+
+	return EMouseCursor::Default;
+}
+
+bool FImageViewportClient::InputKey(const FInputKeyEventArgs& EventArgs)
+{
+	if (EventArgs.Event == IE_Pressed)
+	{
+		if (EventArgs.Key == EKeys::MouseScrollUp || EventArgs.Key == EKeys::Add)
+		{
+			const auto [PixelCoordsValid, PixelCoords] = GetPixelCoordinatesUnderCursor();
+			if (PixelCoordsValid)
+			{
+				Controller.ZoomIn(PixelCoords, GetImageSize.Execute());
+			}
+			return true;
+		}
+
+		if (EventArgs.Key == EKeys::MouseScrollDown || EventArgs.Key == EKeys::Subtract)
+		{
+			const auto [PixelCoordsValid, PixelCoords] = GetPixelCoordinatesUnderCursor();
+			if (PixelCoordsValid)
+			{
+				Controller.ZoomOut(PixelCoords, GetImageSize.Execute());
+			}
+			return true;
+		}
+
+		if (EventArgs.Key == EKeys::F)
+		{
+			Controller.Reset(GetImageSize.Execute(), GetScaledViewportSize());
+
+			return true;
+		}
+	}
+
+	if (OnInputKey.IsBound())
+	{
+		if (OnInputKey.Execute(EventArgs))
+		{
+			return true;
+		}
+	}
+
+	return FEditorViewportClient::InputKey(EventArgs);
+}
+
+void FImageViewportClient::TrackingStarted(const FInputEventState& InputState, bool bIsDraggingWidget, bool bNudge)
+{
+	if (!bIsDraggingWidget)
+	{
+		FIntPoint MousePos;
+		InputState.GetViewport()->GetMousePos(MousePos);
+
+		if (InputState.IsLeftMouseButtonPressed() && MouseIsOverABComparisonDivider(MousePos))
+		{
+			bDraggingABComparisonDivider = true;
+			return;
+		}
+
+		if (InputState.IsRightMouseButtonPressed())
+		{
+			bDragging = true;
+			DraggingStart = MousePos;
+			return;
+		}
+	}
+
+	FEditorViewportClient::TrackingStarted(InputState, bIsDraggingWidget, bNudge);
+}
+
+void FImageViewportClient::TrackingStopped()
+{
+	if (bDraggingABComparisonDivider)
+	{
+		bDraggingABComparisonDivider = false;
+
+		FIntPoint DraggingEnd;
+		Viewport->GetMousePos(DraggingEnd);
+
+		const float LayoutScale = GetLayoutScaleFactor.Execute();
+
+		// Mouse pos is in RT pixel space; convert to unzoomed placement space
+		FVector2d DraggingEndPos = FVector2d(DraggingEnd) / LayoutScale - CachedPlacement.Offset;
+		DraggingEndPos /= CachedPlacement.ZoomFactor;
+
+		const FIntPoint ImageSize = GetImageSize.Execute();
+		if (ImageSize.X > 0)
+		{
+			ABComparisonDivider = DraggingEndPos.X / ImageSize.X;
+		}
+	}
+
+	if (bDragging)
+	{
+		bDragging = false;
+
+		FIntPoint DraggingEnd;
+		Viewport->GetMousePos(DraggingEnd);
+
+		// Drag delta is in RT pixel space; divide by layout and DPI scale to get unzoomed logical space
+		Controller.Pan(FVector2d(DraggingEnd - DraggingStart) / GetLayoutAndDPIScaleFactor.Execute());
+	}
+
+	RequiredCursorVisibiltyAndAppearance.bDontResetCursor = true;
+
+	FEditorViewportClient::TrackingStopped();
+}
+
+int32 FImageViewportClient::GetMipLevel() const
+{
+	return MipLevel;
+}
+
+void FImageViewportClient::SetMipLevel(int32 InMipLevel)
+{
+	MipLevel = InMipLevel;
+}
+
+void FImageViewportClient::ResetController(FIntPoint ImageSize)
+{
+	Controller.Reset(ImageSize, GetScaledViewportSize());
+}
+
+void FImageViewportClient::ResetZoom(FIntPoint ImageSize)
+{
+	const FImageViewportController::FZoomSettings ZoomSettings = Controller.GetZoom();
+	Controller.SetZoom(ZoomSettings.Mode, ZoomSettings.Zoom, ImageSize, GetScaledViewportSize());
+}
+
+FImageViewportController::FZoomSettings FImageViewportClient::GetZoom() const
+{
+	FImageViewportController::FZoomSettings ZoomSettings = Controller.GetZoom();
+	ZoomSettings.Zoom *= GetLayoutAndDPIScaleFactor.Execute();
+	return ZoomSettings;
+}
+
+void FImageViewportClient::SetZoom(FImageViewportController::EZoomMode Mode, double Zoom)
+{
+	Controller.SetZoom(Mode, Zoom, GetImageSize.Execute(), GetScaledViewportSize());
+}
+
+TPair<bool, FVector2d> FImageViewportClient::GetPixelCoordinatesUnderCursor() const
+{
+	if (!bCachedPlacementIsValid || CurrentMousePos == FIntPoint{-1, -1})
+	{
+		return {false, FVector2d::Zero()};
+	}
+
+	const float LayoutAndDPIScale = GetLayoutAndDPIScaleFactor.Execute();
+
+	// Mouse pos is in RT pixel space; convert to unzoomed logical space (where CachedPlacement lives)
+	const FVector2d MousePos((CurrentMousePos.X + 0.5) / LayoutAndDPIScale, (CurrentMousePos.Y + 0.5) / LayoutAndDPIScale);
+	const FVector2d CurrentDrag = GetCurrentDrag();
+
+	FVector2d CursorPos = FVector2d(MousePos) - CachedPlacement.Offset + CurrentDrag;
+	CursorPos /= CachedPlacement.ZoomFactor;
+
+	return {true, CursorPos};
+}
+
+FVector2d FImageViewportClient::GetCurrentDrag() const
+{
+	if (bDragging)
+	{
+		FIntPoint DraggingEnd;
+		Viewport->GetMousePos(DraggingEnd);
+
+		// Drag delta is in RT pixel space; convert to unzoomed logical space
+		return FVector2d(DraggingEnd - DraggingStart) / GetLayoutAndDPIScaleFactor.Execute();
+	}
+	return FVector2d::Zero();
+}
+
+IImageViewer::FDrawProperties::FPlacement FImageViewportClient::GetPlacementProperties(const FIntPoint ImageSize,
+                                                                                       const FVector2d ScaledViewportSize) const
+{
+	const FVector2d CurrentDrag = GetCurrentDrag();
+	const FVector2d Pan = Controller.GetPan(CurrentDrag);
+
+	const double ZoomFactor = Controller.GetZoom().Zoom;
+
+	const FVector2d TileSize = FVector2d(ImageSize) * ZoomFactor;
+	const FVector2d TileOffset = (ScaledViewportSize - TileSize) / 2 + Pan;
+
+	return {TileOffset, TileSize, ZoomFactor};
+}
+
+IImageViewer::FDrawProperties::FMip FImageViewportClient::GetMipProperties() const
+{
+	const double ZoomFactor = Controller.GetZoom().Zoom;
+	const float MipFactor = 1.0f / FMath::Pow(2.0f, MipLevel);
+	const float EffectiveMipLevel = ZoomFactor < MipFactor ? -1.0f : MipLevel;
+
+	return {EffectiveMipLevel};
+}
+
+void FImageViewportClient::CreateOrDestroyCheckerTextureIfSettingsChanged(const SImageViewport::FDrawSettings& DrawSettings)
+{
+	const FCheckerTextureSettings NewCheckerTextureSettings{
+		DrawSettings.bBackgroundCheckerEnabled, DrawSettings.BackgroundCheckerColor1, DrawSettings.BackgroundCheckerColor2, DrawSettings.BackgroundCheckerSize
+	};
+
+	if (CachedCheckerTextureSettings != NewCheckerTextureSettings)
+	{
+		if (CheckerTexture)
+		{
+		    DestroyCheckerTexture(CheckerTexture);
+        }
+
+        if (NewCheckerTextureSettings.bEnabled)
+		{
+			CheckerTexture = TStrongObjectPtr<UTexture2D>(CreateCheckerTexture(NewCheckerTextureSettings.Color1.ToFColorSRGB(),
+			                                                                   NewCheckerTextureSettings.Color2.ToFColorSRGB(),
+			                                                                   NewCheckerTextureSettings.CheckerSize));
+		}
+
+		CachedCheckerTextureSettings = NewCheckerTextureSettings;
+	}
+}
+
+FVector2d FImageViewportClient::GetScaledViewportSize() const
+{
+	// Return unzoomed logical size for controller Reset/SetZoom calls
+	return FVector2d(Viewport->GetSizeXY()) / GetLayoutAndDPIScaleFactor.Execute();
+}
+
+bool FImageViewportClient::MouseIsOverABComparisonDivider(const FIntPoint MousePos) const
+{
+	// Scale the divider position from unzoomed logical space to RT pixel space for comparison with mouse pos
+	const float LayoutScale = GetLayoutScaleFactor.Execute();
+	const double ABComparisonDividerPosition = (CachedPlacement.Offset.X + CachedPlacement.Size.X * ABComparisonDivider) * LayoutScale;
+	const double HitMargin = FMath::Max(4.0, 2.0 * LayoutScale);
+	return ABComparisonDividerPosition - HitMargin <= MousePos.X && MousePos.X <= ABComparisonDividerPosition + HitMargin;
+}
+}
+
+#undef LOCTEXT_NAMESPACE
diff --git a/Engine/Plugins/Editor/ImageWidgets/Source/ImageWidgets/Private/ImageViewportClient.h b/Engine/Plugins/Editor/ImageWidgets/Source/ImageWidgets/Private/ImageViewportClient.h
new file mode 100644
--- /dev/null
+++ b/Engine/Plugins/Editor/ImageWidgets/Source/ImageWidgets/Private/ImageViewportClient.h
@@ -0,0 +1,98 @@
+﻿// Copyright Epic Games, Inc. All Rights Reserved.
+
+#pragma once
+
+#include "EditorViewportClient.h"
+#include "ImageViewportController.h"
+#include "IImageViewer.h"
+#include "SImageViewport.h"
+
+namespace UE::ImageWidgets
+{
+	class FImageABComparison;
+
+	DECLARE_DELEGATE_RetVal(FIntPoint, FGetImageSize)
+	DECLARE_DELEGATE_ThreeParams(FDrawImage, FViewport*, FCanvas*, const IImageViewer::FDrawProperties&)
+	DECLARE_DELEGATE_RetVal(SImageViewport::FDrawSettings, FGetDrawSettings)
+	DECLARE_DELEGATE_RetVal(float, FGetLayoutAndDPIScaleFactor)
+	DECLARE_DELEGATE_RetVal(float, FGetLayoutScaleFactor)
+
+	/**
+	 * Viewport client for controlling the camera and drawing viewport contents. 
+	 */
+	class FImageViewportClient : public FEditorViewportClient
+	{
+	public:
+		FImageViewportClient(const TWeakPtr<SEditorViewport>& InEditorViewport, FGetImageSize&& InGetImageSize, FDrawImage&& InDrawImage,
+		                     FGetDrawSettings&& InGetDrawSettings, FGetLayoutAndDPIScaleFactor&& InGetLayoutAndDPIScaleFactor,
+		                     FGetLayoutScaleFactor&& InGetLayoutScaleFactor, const FImageABComparison* InABComparison,
+		                     const SImageViewport::FControllerSettings& InControllerSettings);
+		virtual ~FImageViewportClient() override;
+
+		virtual void Draw(FViewport* InViewport, FCanvas* Canvas) override;
+
+		// FEditorViewportClient overrides - begin
+		virtual EMouseCursor::Type GetCursor(FViewport* InViewport, int32 X, int32 Y) override;
+		virtual bool InputKey(const FInputKeyEventArgs& EventArgs) override;
+		virtual void TrackingStarted(const FInputEventState& InputState, bool bIsDraggingWidget, bool bNudge) override;
+		virtual void TrackingStopped() override;
+		// FEditorViewportClient overrides - end
+
+		int32 GetMipLevel() const;
+		void SetMipLevel(int32 InMipLevel);
+		FImageViewportController::FZoomSettings GetZoom() const;
+		void SetZoom(FImageViewportController::EZoomMode Mode, double Zoom);
+		void ResetController(FIntPoint ImageSize);
+		void ResetZoom(FIntPoint ImageSize);
+
+		TPair<bool, FVector2d> GetPixelCoordinatesUnderCursor() const;
+
+	private:
+		struct FCheckerTextureSettings
+		{
+			bool operator==(const FCheckerTextureSettings&) const = default;
+
+			bool bEnabled = false;
+			FLinearColor Color1;
+			FLinearColor Color2;
+			uint32 CheckerSize = 0;
+		};
+
+		FVector2d GetCurrentDrag() const;
+		IImageViewer::FDrawProperties::FPlacement GetPlacementProperties(FIntPoint ImageSize, FVector2d ScaledViewportSize) const;
+		IImageViewer::FDrawProperties::FMip GetMipProperties() const;
+
+		void CreateOrDestroyCheckerTextureIfSettingsChanged(const SImageViewport::FDrawSettings& DrawSettings);
+		bool MouseIsOverABComparisonDivider(FIntPoint MousePos) const;
+
+		FVector2d GetScaledViewportSize() const;
+
+		FGetImageSize GetImageSize;
+		FDrawImage DrawImage;
+		FGetDrawSettings GetDrawSettings;
+		FGetLayoutAndDPIScaleFactor GetLayoutAndDPIScaleFactor;
+		FGetLayoutScaleFactor GetLayoutScaleFactor;
+		SImageViewport::FControllerSettings::FOnInputKey OnInputKey;
+		const FImageABComparison* ABComparison;
+
+		bool bDragging = false;
+		FIntPoint DraggingStart;
+
+		bool bMipSelected = true;
+		int32 MipLevel = INDEX_NONE;
+
+		IImageViewer::FDrawProperties::FPlacement CachedPlacement;
+		bool bCachedPlacementIsValid = false;
+
+		FImageViewportController Controller;
+
+		TStrongObjectPtr<UTexture2D> CheckerTexture;
+		FCheckerTextureSettings CachedCheckerTextureSettings;
+
+		FVector2d CachedScaledViewportSize = FVector2d::ZeroVector;
+		bool bZoomOnResize = false;
+
+		double ABComparisonDivider = 0.5;
+		bool bDraggingABComparisonDivider = false;
+	};
+}
diff --git a/Engine/Plugins/Editor/ImageWidgets/Source/ImageWidgets/Private/ImageViewportController.cpp b/Engine/Plugins/Editor/ImageWidgets/Source/ImageWidgets/Private/ImageViewportController.cpp
new file mode 100644
--- /dev/null
+++ b/Engine/Plugins/Editor/ImageWidgets/Source/ImageWidgets/Private/ImageViewportController.cpp
@@ -0,0 +1,161 @@
+// Copyright Epic Games, Inc. All Rights Reserved.
+
+#include "ImageViewportController.h"
+
+namespace UE::ImageWidgets
+{
+	namespace ImageViewportController_Local
+	{
+		const TArray<double>& GetZoomLevels()
+		{
+			constexpr double ZoomMin = 1.0 / 64;
+			constexpr double ZoomMax = 64.0;
+			constexpr int32 ZoomStepsToDouble = 8;
+			constexpr double ZoomPowerIncrement = 1.0 / ZoomStepsToDouble;
+
+			static TArray<double> ZoomLevels = []
+			{
+				const int32 ZoomMinPower = FMath::RoundToInt32(FMath::Log2(ZoomMin));
+				const int32 ZoomMaxPower = FMath::RoundToInt32(FMath::Log2(ZoomMax));
+				const int32 Num = (ZoomMaxPower - ZoomMinPower) * ZoomStepsToDouble + 1;
+
+				TArray<double> Result;
+				Result.Reserve(Num);
+
+				for (double Power = ZoomMinPower; Power <= ZoomMaxPower; Power += ZoomPowerIncrement)
+				{
+					Result.Add(FMath::Pow(2.0, Power));
+				}
+
+				return Result;
+			}();
+
+			return ZoomLevels;
+		}
+
+		double ZoomIn(const double CurrentZoom)
+		{
+			const TArray<double>& ZoomLevels = GetZoomLevels();
+			const int32 Index = Algo::LowerBound(ZoomLevels, CurrentZoom);
+
+			if (Index >= ZoomLevels.Num())
+			{
+				return ZoomLevels.Last();
+			}
+
+			if (CurrentZoom < ZoomLevels[Index])
+			{
+				return ZoomLevels[Index];
+			}
+
+			return ZoomLevels[FMath::Min(Index + 1, ZoomLevels.Num() - 1)];
+		}
+
+		double ZoomOut(const double CurrentZoom)
+		{
+			const TArray<double>& ZoomLevels = GetZoomLevels();
+			const int32 Index = Algo::LowerBound(ZoomLevels, CurrentZoom);
+
+			if (Index >= ZoomLevels.Num())
+			{
+				return ZoomLevels.Last();
+			}
+
+			if (CurrentZoom > ZoomLevels[Index])
+			{
+				return ZoomLevels[Index];
+			}
+
+			return ZoomLevels[FMath::Max(0, Index - 1)];
+		}
+
+		FVector2d AdjustPan(const FVector2d OldPan, const double OldZoom, const double NewZoom, const FVector2d CursorPos, const FVector2d& ImageSize)
+		{
+			const FVector2d ImageCenter = ImageSize * 0.5;
+			const FVector2d DistanceCursorToCenter = CursorPos - ImageCenter;
+			const FVector2d DeltaDistanceCursorToCenter = DistanceCursorToCenter * (NewZoom - OldZoom);
+			const FVector2d NewPan = OldPan - DeltaDistanceCursorToCenter;
+
+			return NewPan;
+		}
+	}
+
+	FImageViewportController::FImageViewportController(const EZoomMode DefaultZoomMode)
+		: DefaultZoomMode(DefaultZoomMode)
+	{
+		Reset(FIntPoint::ZeroValue, FVector2d::Zero());
+	}
+
+	void FImageViewportController::Pan(const FVector2d ViewportDeltaWithDPIScaling)
+	{
+		PanAmount += ViewportDeltaWithDPIScaling;
+	}
+
+	void FImageViewportController::Reset(const FIntPoint ImageSize, const FVector2d ViewportSizeWithDPIScaling)
+	{
+		SetZoom(DefaultZoomMode, 1.0, ImageSize, ViewportSizeWithDPIScaling);
+
+		PanAmount = FVector2d::Zero();
+	}
+
+	void FImageViewportController::ZoomIn(const FVector2d CursorPixelPos, const FIntPoint ImageSize)
+	{
+		const double OldZoom = ZoomSettings.Zoom;
+
+		ZoomSettings.Mode = EZoomMode::Custom;
+		ZoomSettings.Zoom = ImageViewportController_Local::ZoomIn(ZoomSettings.Zoom);
+
+		PanAmount = ImageViewportController_Local::AdjustPan(PanAmount, OldZoom, ZoomSettings.Zoom, CursorPixelPos, ImageSize);
+	}
+
+	void FImageViewportController::ZoomOut(const FVector2d CursorPixelPos, const FIntPoint ImageSize)
+	{
+		const double OldZoom = ZoomSettings.Zoom;
+
+		ZoomSettings.Mode = EZoomMode::Custom;
+		ZoomSettings.Zoom = ImageViewportController_Local::ZoomOut(ZoomSettings.Zoom);
+
+		PanAmount = ImageViewportController_Local::AdjustPan(PanAmount, OldZoom, ZoomSettings.Zoom, CursorPixelPos, ImageSize);
+	}
+
+	FVector2d FImageViewportController::GetPan(const FVector2d Drag) const
+	{
+		return PanAmount + Drag;
+	}
+
+	FImageViewportController::FZoomSettings FImageViewportController::GetZoom() const
+	{
+		return ZoomSettings;
+	}
+
+	void FImageViewportController::SetZoom(const EZoomMode ZoomMode, const double Zoom, const FIntPoint ImageSize, const FVector2d ViewportSizeWithDPIScaling)
+	{
+		ZoomSettings.Mode = ZoomMode;
+
+		if (ZoomSettings.Mode == EZoomMode::Custom)
+		{
+			ZoomSettings.Zoom = Zoom;
+			return;
+		}
+
+		if (ImageSize.X <= 0 || ImageSize.Y <= 0 || ViewportSizeWithDPIScaling == FVector2d::Zero())
+		{
+			ZoomSettings.Zoom = 1.0;
+			return;
+		}
+
+		const FVector2d SizeRatio = ViewportSizeWithDPIScaling / ImageSize;
+		const double SizeRatioMin = SizeRatio.GetMin();
+
+		if (ZoomSettings.Mode == EZoomMode::Fill)
+		{
+			ZoomSettings.Zoom = SizeRatioMin;
+		}
+		else
+		{
+			ZoomSettings.Zoom = FMath::Min(1.0, SizeRatioMin);
+		}
+
+		PanAmount = FVector2d::Zero();
+	}
+}
diff --git a/Engine/Plugins/Editor/ImageWidgets/Source/ImageWidgets/Private/ImageViewportController.h b/Engine/Plugins/Editor/ImageWidgets/Source/ImageWidgets/Private/ImageViewportController.h
new file mode 100644
--- /dev/null
+++ b/Engine/Plugins/Editor/ImageWidgets/Source/ImageWidgets/Private/ImageViewportController.h
@@ -0,0 +1,48 @@
+// Copyright Epic Games, Inc. All Rights Reserved.
+
+#pragma once
+
+#include "CoreMinimal.h"
+#include "SImageViewport.h"
+
+namespace UE::ImageWidgets
+{
+	/**
+	 * Camera controller for the 2D viewport supporting panning and zooming.
+	 *
+	 * Note: All function parameters relating to the viewport need to be passed in with DPI scaling taken into account!
+	 */
+	class FImageViewportController
+	{
+	public:
+		enum class EZoomMode
+		{
+			Fit = static_cast<int32>(SImageViewport::FControllerSettings::EDefaultZoomMode::Fit),
+			Fill = static_cast<int32>(SImageViewport::FControllerSettings::EDefaultZoomMode::Fill),
+			Custom
+		};
+		static_assert(EZoomMode::Fit != EZoomMode::Fill && EZoomMode::Fit != EZoomMode::Custom && EZoomMode::Fill != EZoomMode::Custom);
+
+		struct FZoomSettings
+		{
+			EZoomMode Mode;
+			double Zoom;
+		};
+
+		explicit FImageViewportController(EZoomMode DefaultZoomMode);
+
+		void Pan(FVector2d ViewportDeltaWithDPIScaling);
+		void Reset(FIntPoint ImageSize, FVector2d ViewportSizeWithDPIScaling);
+		void ZoomIn(FVector2d CursorPixelPos, FIntPoint ImageSize);
+		void ZoomOut(FVector2d CursorPixelPos, FIntPoint ImageSize);
+
+		FVector2d GetPan(FVector2d Drag) const;
+		FZoomSettings GetZoom() const;
+		void SetZoom(EZoomMode ZoomMode, double Zoom, FIntPoint ImageSize, FVector2d ViewportSizeWithDPIScaling);
+
+	private:
+		FVector2d PanAmount;
+		FZoomSettings ZoomSettings;
+		EZoomMode DefaultZoomMode;
+	};
+}
diff --git a/Engine/Plugins/Editor/ImageWidgets/Source/ImageWidgets/Private/ImageWidgetsCommands.cpp b/Engine/Plugins/Editor/ImageWidgets/Source/ImageWidgets/Private/ImageWidgetsCommands.cpp
new file mode 100644
--- /dev/null
+++ b/Engine/Plugins/Editor/ImageWidgets/Source/ImageWidgets/Private/ImageWidgetsCommands.cpp
@@ -0,0 +1,37 @@
+﻿// Copyright Epic Games, Inc. All Rights Reserved.
+
+#include "ImageWidgetsCommands.h"
+
+#include "Framework/Commands/InputChord.h"
+#include "Styling/AppStyle.h"
+
+#define LOCTEXT_NAMESPACE "ImageWidgetsCommands"
+
+namespace UE::ImageWidgets
+{
+	FImageWidgetsCommands::FImageWidgetsCommands()
+		: TCommands("ImageWidgets", LOCTEXT("ContextDescription", "Image Widgets"), NAME_None, FAppStyle::Get().GetStyleSetName())
+	{
+	}
+
+	void FImageWidgetsCommands::RegisterCommands()
+	{
+		UI_COMMAND(MipMinus, "Lower", "Show next lower mip", EUserInterfaceActionType::None, FInputChord());
+		UI_COMMAND(MipPlus, "Higher", "Show next higher mip", EUserInterfaceActionType::None, FInputChord());
+
+		UI_COMMAND(ToggleOverlay, "Toggle Overlay", "Toggles the visibility of the toolbar and the status bar.", EUserInterfaceActionType::ToggleButton,
+		           FInputChord(EKeys::O));
+
+		UI_COMMAND(Zoom12, "12.5%", "Set Zoom to 12.5%", EUserInterfaceActionType::Button, FInputChord());
+		UI_COMMAND(Zoom25, "25%", "Set Zoom to 25%", EUserInterfaceActionType::Button, FInputChord());
+		UI_COMMAND(Zoom50, "50%", "Set Zoom to 50%", EUserInterfaceActionType::Button, FInputChord());
+		UI_COMMAND(Zoom100, "100%", "Set Zoom to 100%", EUserInterfaceActionType::Button, FInputChord());
+		UI_COMMAND(Zoom200, "200%", "Set Zoom to 200%", EUserInterfaceActionType::Button, FInputChord());
+		UI_COMMAND(Zoom400, "400%", "Set Zoom to 400%", EUserInterfaceActionType::Button, FInputChord());
+		UI_COMMAND(Zoom800, "800%", "Set Zoom to 800%", EUserInterfaceActionType::Button, FInputChord());
+		UI_COMMAND(ZoomFit, "Scale to Fit", "Optionally downscales the image to fit within the viewport", EUserInterfaceActionType::RadioButton, FInputChord());
+		UI_COMMAND(ZoomFill, "Scale to Fill", "Downscales or upscales the image to fill the viewport", EUserInterfaceActionType::RadioButton, FInputChord());
+	}
+}
+
+#undef LOCTEXT_NAMESPACE
diff --git a/Engine/Plugins/Editor/ImageWidgets/Source/ImageWidgets/Private/ImageWidgetsCommands.h b/Engine/Plugins/Editor/ImageWidgets/Source/ImageWidgets/Private/ImageWidgetsCommands.h
new file mode 100644
--- /dev/null
+++ b/Engine/Plugins/Editor/ImageWidgets/Source/ImageWidgets/Private/ImageWidgetsCommands.h
@@ -0,0 +1,35 @@
+﻿// Copyright Epic Games, Inc. All Rights Reserved.
+
+#pragma once
+
+#include "Framework/Commands/Commands.h"
+#include "Framework/Commands/UICommandInfo.h"
+
+namespace UE::ImageWidgets
+{
+	/**
+	 * Provides commands used for the image widgets.
+	 */
+	class FImageWidgetsCommands : public TCommands<FImageWidgetsCommands>
+	{
+	public:
+		FImageWidgetsCommands();
+
+		virtual void RegisterCommands() override;
+
+		TSharedPtr<FUICommandInfo> MipMinus;
+		TSharedPtr<FUICommandInfo> MipPlus;
+
+		TSharedPtr<FUICommandInfo> ToggleOverlay;
+
+		TSharedPtr<FUICommandInfo> Zoom12;
+		TSharedPtr<FUICommandInfo> Zoom25;
+		TSharedPtr<FUICommandInfo> Zoom50;
+		TSharedPtr<FUICommandInfo> Zoom100;
+		TSharedPtr<FUICommandInfo> Zoom200;
+		TSharedPtr<FUICommandInfo> Zoom400;
+		TSharedPtr<FUICommandInfo> Zoom800;
+		TSharedPtr<FUICommandInfo> ZoomFit;
+		TSharedPtr<FUICommandInfo> ZoomFill;
+	};
+}
diff --git a/Engine/Plugins/Editor/ImageWidgets/Source/ImageWidgets/Private/ImageWidgetsLogCategory.cpp b/Engine/Plugins/Editor/ImageWidgets/Source/ImageWidgets/Private/ImageWidgetsLogCategory.cpp
new file mode 100644
--- /dev/null
+++ b/Engine/Plugins/Editor/ImageWidgets/Source/ImageWidgets/Private/ImageWidgetsLogCategory.cpp
@@ -0,0 +1,8 @@
+﻿// Copyright Epic Games, Inc. All Rights Reserved.
+
+#include "ImageWidgetsLogCategory.h"
+
+namespace UE::ImageWidgets
+{
+	DEFINE_LOG_CATEGORY(LogImageWidgets);
+}
diff --git a/Engine/Plugins/Editor/ImageWidgets/Source/ImageWidgets/Private/ImageWidgetsLogCategory.h b/Engine/Plugins/Editor/ImageWidgets/Source/ImageWidgets/Private/ImageWidgetsLogCategory.h
new file mode 100644
--- /dev/null
+++ b/Engine/Plugins/Editor/ImageWidgets/Source/ImageWidgets/Private/ImageWidgetsLogCategory.h
@@ -0,0 +1,10 @@
+﻿// Copyright Epic Games, Inc. All Rights Reserved.
+
+#pragma once
+
+#include "Logging/LogMacros.h"
+
+namespace UE::ImageWidgets
+{
+	DECLARE_LOG_CATEGORY_EXTERN(LogImageWidgets, Log, All)
+}
diff --git a/Engine/Plugins/Editor/ImageWidgets/Source/ImageWidgets/Private/ImageWidgetsModule.cpp b/Engine/Plugins/Editor/ImageWidgets/Source/ImageWidgets/Private/ImageWidgetsModule.cpp
new file mode 100644
--- /dev/null
+++ b/Engine/Plugins/Editor/ImageWidgets/Source/ImageWidgets/Private/ImageWidgetsModule.cpp
@@ -0,0 +1,116 @@
+﻿// Copyright Epic Games, Inc. All Rights Reserved.
+
+#if !IMAGE_WIDGETS_BUILD_COLOR_VIEWER_SAMPLE
+
+#include "Modules/ModuleManager.h"
+
+IMPLEMENT_MODULE(FDefaultModuleImpl, ImageWidgets);
+
+#else // !IMAGE_WIDGETS_BUILD_COLOR_VIEWER_SAMPLE
+
+#include "CoreMinimal.h"
+#include "ColorViewerSample/ColorViewerCommands.h"
+#include "ColorViewerSample/ColorViewerWidget.h"
+#include "WorkspaceMenuStructure.h"
+#include "WorkspaceMenuStructureModule.h"
+#include "Framework/Docking/TabManager.h"
+#include "Modules/ModuleInterface.h"
+#include "Widgets/Docking/SDockTab.h"
+
+#define LOCTEXT_NAMESPACE "ImageWidgetsModule"
+
+namespace UE::ImageWidgets
+{
+	/**
+	 * This module provides the color viewer sample within a tab widget accessible via the "Tools > Miscellaneous" menu.
+	 */
+	class FImageWidgetsModule : public IModuleInterface
+	{
+	public:
+		void RegisterTabSpawners(const TSharedPtr<FWorkspaceItem>& WorkspaceGroup);
+		void UnregisterTabSpawners();
+
+		// IModuleInterface overrides - begin
+		virtual void StartupModule() override;
+		virtual void ShutdownModule() override;
+		// IModuleInterface overrides - end
+
+	private:
+		TSharedRef<SDockTab> MakeColorViewerTab(const FSpawnTabArgs&);
+		TSharedRef<SWidget> GetColorViewer();
+
+		bool bHasRegisteredTabSpawners = false;
+		TWeakPtr<Sample::SColorViewerWidget> ColorViewerPtr;
+	};
+
+	void FImageWidgetsModule::RegisterTabSpawners(const TSharedPtr<FWorkspaceItem>& WorkspaceGroup)
+	{
+		if (bHasRegisteredTabSpawners)
+		{
+			UnregisterTabSpawners();
+		}
+
+		bHasRegisteredTabSpawners = true;
+
+		FTabSpawnerEntry& ColorViewerSpawnerEntry = FGlobalTabmanager::Get()->RegisterNomadTabSpawner("ColorViewer",
+			FOnSpawnTab::CreateRaw(this, &FImageWidgetsModule::MakeColorViewerTab))
+		.SetDisplayName(LOCTEXT("ColorViewerTitle", "Color Viewer Sample"))
+		.SetTooltipText(LOCTEXT("ColorViewerTooltipText", "Open the Color Viewer tab, a sample application for the ImageWidgets plugin."))
+		.SetGroup(WorkspaceMenu::GetMenuStructure().GetDeveloperToolsMiscCategory())
+		.SetIcon(FSlateIcon(FAppStyle::GetAppStyleSetName(), "Icons.Color"));
+
+		if (WorkspaceGroup.IsValid())
+		{
+			ColorViewerSpawnerEntry.SetGroup(WorkspaceGroup.ToSharedRef());
+		}
+	}
+
+	void FImageWidgetsModule::UnregisterTabSpawners()
+	{
+		bHasRegisteredTabSpawners = false;
+
+		FGlobalTabmanager::Get()->UnregisterNomadTabSpawner("ColorViewer");
+	}
+
+	void FImageWidgetsModule::StartupModule()
+	{
+		bHasRegisteredTabSpawners = false;
+		RegisterTabSpawners(nullptr);
+
+		Sample::FColorViewerCommands::Register();
+	}
+
+	void FImageWidgetsModule::ShutdownModule()
+	{
+		Sample::FColorViewerCommands::Unregister();
+
+		UnregisterTabSpawners();
+	}
+
+	TSharedRef<SDockTab> FImageWidgetsModule::MakeColorViewerTab(const FSpawnTabArgs&)
+	{
+		TSharedRef<SDockTab> Tab = SNew(SDockTab)
+			.TabRole(NomadTab);
+		Tab->SetContent(GetColorViewer());
+		return Tab;
+	}
+
+	TSharedRef<SWidget> FImageWidgetsModule::GetColorViewer()
+	{
+		TSharedPtr<Sample::SColorViewerWidget> ColorViewer = ColorViewerPtr.Pin();
+
+		if (!ColorViewer.IsValid())
+		{
+			ColorViewer = SNew(Sample::SColorViewerWidget);
+			ColorViewerPtr = ColorViewer;
+		}
+
+		return ColorViewer.ToSharedRef();
+	}
+}
+
+#undef LOCTEXT_NAMESPACE
+
+IMPLEMENT_MODULE(UE::ImageWidgets::FImageWidgetsModule, ImageWidgets)
+
+#endif // !IMAGE_WIDGETS_BUILD_COLOR_VIEWER_SAMPLE
diff --git a/Engine/Plugins/Editor/ImageWidgets/Source/ImageWidgets/Private/ImageWidgetsStyle.cpp b/Engine/Plugins/Editor/ImageWidgets/Source/ImageWidgets/Private/ImageWidgetsStyle.cpp
new file mode 100644
--- /dev/null
+++ b/Engine/Plugins/Editor/ImageWidgets/Source/ImageWidgets/Private/ImageWidgetsStyle.cpp
@@ -0,0 +1,42 @@
+﻿// Copyright Epic Games, Inc. All Rights Reserved.
+
+#include "ImageWidgetsStyle.h"
+
+#include "Styling/AppStyle.h"
+#include "Styling/CoreStyle.h"
+#include "Styling/SlateStyleRegistry.h"
+#include "Styling/SlateTypes.h"
+
+namespace UE::ImageWidgets
+{
+	FName FImageWidgetsStyle::StyleName("ImageWidgetsStyle");
+
+	FImageWidgetsStyle::FImageWidgetsStyle()
+		: FSlateStyleSet(StyleName)
+	{
+		FSlateStyleSet::SetContentRoot(FPaths::EngineContentDir() / TEXT("Editor/Slate"));
+		FSlateStyleSet::SetCoreContentRoot(FPaths::EngineContentDir() / TEXT("Slate"));
+
+		const FTextBlockStyle NormalText = FAppStyle::Get().GetWidgetStyle<FTextBlockStyle>("NormalText");
+		Set("RichTextBlock.Red", FTextBlockStyle(NormalText).SetColorAndOpacity(FLinearColor(1.0f, 0.1f, 0.1f)));
+		Set("RichTextBlock.Green", FTextBlockStyle(NormalText).SetColorAndOpacity(FLinearColor(0.1f, 1.0f, 0.1f)));
+		Set("RichTextBlock.Blue", FTextBlockStyle(NormalText).SetColorAndOpacity(FLinearColor(0.1f, 0.1f, 1.0f)));
+
+		FSlateBrush* BrushTableRowOdd = new FSlateBrush();
+		BrushTableRowOdd->TintColor = FLinearColor(0.5f, 0.5f, 0.5f, 1.0f);
+		Set("TableRowOdd", BrushTableRowOdd);
+
+		FSlateStyleRegistry::RegisterSlateStyle(*this);
+	}
+
+	FImageWidgetsStyle::~FImageWidgetsStyle()
+	{
+		FSlateStyleRegistry::UnRegisterSlateStyle(*this);
+	}
+
+	FImageWidgetsStyle& FImageWidgetsStyle::Get()
+	{
+		static FImageWidgetsStyle Instance;
+		return Instance;
+	}
+}
diff --git a/Engine/Plugins/Editor/ImageWidgets/Source/ImageWidgets/Private/ImageWidgetsStyle.h b/Engine/Plugins/Editor/ImageWidgets/Source/ImageWidgets/Private/ImageWidgetsStyle.h
new file mode 100644
--- /dev/null
+++ b/Engine/Plugins/Editor/ImageWidgets/Source/ImageWidgets/Private/ImageWidgetsStyle.h
@@ -0,0 +1,23 @@
+﻿// Copyright Epic Games, Inc. All Rights Reserved.
+
+#pragma once
+
+#include "Styling/SlateStyle.h"
+
+namespace UE::ImageWidgets
+{
+	/**
+	 * Style declarations for the image widgets.
+	 */
+	class FImageWidgetsStyle : public FSlateStyleSet
+	{
+	public:
+		static FName StyleName;
+
+		static FImageWidgetsStyle& Get();
+
+	private:
+		FImageWidgetsStyle();
+		virtual ~FImageWidgetsStyle() override;
+	};
+}
diff --git a/Engine/Plugins/Editor/ImageWidgets/Source/ImageWidgets/Private/SImageCatalog.cpp b/Engine/Plugins/Editor/ImageWidgets/Source/ImageWidgets/Private/SImageCatalog.cpp
new file mode 100644
--- /dev/null
+++ b/Engine/Plugins/Editor/ImageWidgets/Source/ImageWidgets/Private/SImageCatalog.cpp
@@ -0,0 +1,1000 @@
+﻿// Copyright Epic Games, Inc. All Rights Reserved.
+
+#include "SImageCatalog.h"
+
+#include "Algo/Accumulate.h"
+#include "Brushes/SlateColorBrush.h"
+#include "Framework/Application/SlateApplication.h"
+#include "Styling/StyleColors.h"
+#include "Widgets/Layout/SScrollBox.h"
+#include "Widgets/Views/SListView.h"
+#include "Widgets/Input/SButton.h"
+
+#include "SImageCatalogItem.h"
+
+#define LOCTEXT_NAMESPACE "SImageViewerCatalog"
+
+namespace UE::ImageWidgets
+{
+FImageCatalogItemData::FImageCatalogItemData(const FGuid Guid, const FSlateBrush& Brush, const FText& Name, const FText& Info, const FText& ToolTip)
+	: Guid(Guid), Thumbnail(Brush), Name(Name), Info(Info), ToolTip(ToolTip)
+{
+}
+
+using FItemType = TSharedPtr<FImageCatalogItemData>;
+
+struct FGroup
+{
+	FGroup(FName Name, FText Heading)
+		: Name(Name), Heading(MoveTemp(Heading))
+	{}
+
+	FName Name;
+	FText Heading;
+	TSharedPtr<SListView<FItemType>> ListView;
+	TArray<FItemType> Items;
+	bool bIsExpanded = true;
+};
+
+struct FItemLookup
+{
+	int32 GroupIndex;
+	int32 ItemIndex;
+};
+
+class FImpl
+{
+public:
+	FImpl(FName DefaultGroupName, const FText& DefaultGroupHeading, ESelectionMode::Type SelectionMode, bool bAllowSelectionAcrossGroups, bool bShowEmptyGroups,
+	      const SImageCatalog::FOnItemSelected& OnItemSelected, const SImageCatalog::FOnGetGroupContextMenu& OnGetGroupContextMenu,
+	      const SImageCatalog::FOnGetItemsContextMenu& OnGetItemsContextMenu, TSharedPtr<SScrollBox> Layout);
+
+	FName GetDefaultGroupName() const { return DefaultGroupName; }
+	bool AddGroup(FName Name, const FText& Heading, const FName* BeforeGroupWithThisName);
+	TOptional<TArray<FGuid>> RemoveGroup(FName Name, const FName* GroupToMoveItemsInto);
+	int32 NumGroups() const;
+	bool SetGroupHeading(FName Name, const FText& Heading);
+	TOptional<FName> GetGroupNameAt(int32 Index) const;
+
+	bool AddItem(const FItemType& Item, const FName* GroupName, const FGuid* BeforeItemWithThisGuid);
+	bool MoveItem(const FGuid& Guid, const FName* GroupName, const FGuid* BeforeItemWithThisGuid);
+	bool RemoveItem(const FGuid& Guid);
+	TSharedPtr<const FImageCatalogItemData> GetItem(const FGuid& Guid) const;
+	TOptional<int32> GetItemIndex(const FGuid& Guid) const;
+	TOptional<TTuple<FName, int32>> GetItemGroupNameAndIndex(const FGuid& Guid) const;
+	TSharedPtr<const FImageCatalogItemData> GetItemAt(int32 Index, const FName* GroupName) const;
+	TOptional<FGuid> GetItemGuidAt(int32 Index, const FName* GroupName) const;
+	bool UpdateItem(const FGuid& Guid, const FSlateBrush* Thumbnail, const FText* Name, const FText* Info, const FText* ToolTip);
+
+	bool SelectItem(const FGuid& Guid, bool bSelected);
+	bool ClearSelection(const FName* GroupName);
+
+	int32 NumItems(const FName* GroupName) const;
+	TOptional<FName> GetItemGroupName(const FGuid& Guid) const;
+
+private:
+	bool AddGroup(FName Name, const FText& Heading, int32 Index);
+
+	const FGroup* FindGroup(FName Name) const;
+	FGroup* FindGroup(FName Name)
+	{
+		return const_cast<FGroup*>(const_cast<const FImpl*>(this)->FindGroup(Name));
+	}
+
+	const FItemLookup* FindLookup(const FGuid& Guid) const;
+	FItemLookup* FindLookup(const FGuid& Guid)
+	{
+		return const_cast<FItemLookup*>(const_cast<const FImpl*>(this)->FindLookup(Guid));
+	}
+
+	TSharedPtr<SWidget> OnGroupContextMenuOpening(FName Name) const;
+	TSharedPtr<SWidget> OnItemsContextMenuOpening() const;
+	void SortSelection(TArray<FGuid>& Selection) const;
+	void UpdateGroupMapping(int32 StartingIndex);
+	void UpdateItemMapping(TArray<FItemType>& Items, int32 StartingIndex);
+
+	const FName DefaultGroupName;
+	const FText DefaultGroupHeading;
+	const ESelectionMode::Type SelectionMode;
+	const bool bAllowSelectionAcrossGroups;
+	const bool bShowEmptyGroups;
+	SImageCatalog::FOnItemSelected OnItemSelected;
+	SImageCatalog::FOnGetGroupContextMenu OnGetGroupContextMenu;
+	SImageCatalog::FOnGetItemsContextMenu OnGetItemsContextMenu;
+	TSharedPtr<SScrollBox> Layout;
+	TArray<TUniquePtr<FGroup>> Groups;
+	TMap<FName, int32> GroupMapping;
+	TMap<FGuid, FItemLookup> ItemMapping;
+};
+
+FImpl::FImpl(const FName DefaultGroupName, const FText& DefaultGroupHeading, const ESelectionMode::Type SelectionMode, const bool bAllowSelectionAcrossGroups,
+             const bool bShowEmptyGroups, const SImageCatalog::FOnItemSelected& OnItemSelected,
+             const SImageCatalog::FOnGetGroupContextMenu& OnGetGroupContextMenu, const SImageCatalog::FOnGetItemsContextMenu& OnGetItemsContextMenu,
+             TSharedPtr<SScrollBox> Layout)
+	: DefaultGroupName(DefaultGroupName)
+	, DefaultGroupHeading(DefaultGroupHeading)
+	, SelectionMode(SelectionMode)
+	, bAllowSelectionAcrossGroups(bAllowSelectionAcrossGroups)
+	, bShowEmptyGroups(bShowEmptyGroups)
+	, OnItemSelected(OnItemSelected)
+	, OnGetGroupContextMenu(OnGetGroupContextMenu)
+	, OnGetItemsContextMenu(OnGetItemsContextMenu)
+	, Layout(MoveTemp(Layout))
+{
+	AddGroup(DefaultGroupName, DefaultGroupHeading, nullptr);
+}
+
+bool FImpl::AddGroup(const FName Name, const FText& Heading, const FName* BeforeGroupWithThisName)
+{
+	if (Layout.IsValid())
+	{
+		if (BeforeGroupWithThisName != nullptr)
+		{
+			if (const int32 *const Index = GroupMapping.Find(*BeforeGroupWithThisName))
+			{
+				return AddGroup(Name, Heading, *Index);
+			}
+		}
+
+		return AddGroup(Name, Heading, Groups.Num());
+	}
+	return false;
+}
+
+TOptional<TArray<FGuid>> FImpl::RemoveGroup(const FName Name, const FName* GroupToMoveItemsInto)
+{
+	if (Name == DefaultGroupName)
+	{
+		return {};
+	}
+
+	if (GroupToMoveItemsInto && Name == *GroupToMoveItemsInto)
+	{
+		return {};
+	}
+
+	if (const int32* GroupIndexPtr = GroupMapping.Find(Name))
+	{
+		const int32 GroupIndex = *GroupIndexPtr;
+		check(0 <= GroupIndex && GroupIndex < Groups.Num());
+		FGroup& Group = *Groups[GroupIndex];
+		TArray<FItemType>& GroupItems = Group.Items;
+
+		TArray<FGuid> AffectedGuids;
+		AffectedGuids.Reserve(GroupItems.Num());
+
+		const int32* NewGroupIndex = GroupToMoveItemsInto ? GroupMapping.Find(*GroupToMoveItemsInto) : nullptr;
+		if (GroupToMoveItemsInto && !NewGroupIndex)
+		{
+			NewGroupIndex = GroupMapping.Find(DefaultGroupName);
+		}
+		if (NewGroupIndex)
+		{
+			check(0 <= *NewGroupIndex && *NewGroupIndex < Groups.Num());
+			FGroup& NewGroup = *Groups[*NewGroupIndex];
+			TArray<FItemType>& NewGroupItems = NewGroup.Items;
+
+			for (FItemType& Item : GroupItems)
+			{
+				AffectedGuids.Add(Item->Guid);
+
+				FItemLookup* const Lookup = FindLookup(Item->Guid);
+				check(Lookup);
+				Lookup->GroupIndex = *NewGroupIndex;
+				Lookup->ItemIndex = NewGroupItems.Num();
+
+				NewGroupItems.Emplace(MoveTemp(Item));
+			}
+
+			NewGroup.ListView->RequestListRefresh();
+		}
+		else
+		{
+			for (const FItemType& Item : GroupItems)
+			{
+				AffectedGuids.Add(Item->Guid);
+				ItemMapping.Remove(Item->Guid);
+			}
+		}
+
+		check(Layout->GetChildren() && GroupIndex < Layout->NumSlots());
+		const TSharedRef<SWidget>& WidgetToRemove = Layout->GetSlot(GroupIndex).GetWidget();
+		Layout->RemoveSlot(WidgetToRemove);
+		Groups.RemoveAt(GroupIndex);
+		GroupMapping.Remove(Name);
+		UpdateGroupMapping(GroupIndex);
+
+		return {AffectedGuids};
+	}
+
+	return {};
+}
+
+int32 FImpl::NumGroups() const
+{
+	return Groups.Num(); 
+}
+
+bool FImpl::SetGroupHeading(const FName Name, const FText& Heading)
+{
+	if (FGroup *const Group = FindGroup(Name))
+	{
+		Group->Heading = Heading;
+		return true;
+	}
+
+	return false;
+}
+
+TOptional<FName> FImpl::GetGroupNameAt(const int32 Index) const
+{
+	if (0 <= Index && Index < Groups.Num())
+	{
+		return {Groups[Index]->Name};
+	}
+
+	return {};
+}
+
+bool FImpl::AddItem(const FItemType& Item, const FName* GroupName, const FGuid* BeforeItemWithThisGuid)
+{
+	if (GroupName == nullptr)
+	{
+		GroupName = &DefaultGroupName;
+	}
+
+	if (const int32* GroupIndex = GroupMapping.Find(*GroupName))
+	{
+		check(0 <= *GroupIndex && *GroupIndex < Groups.Num());
+		FGroup& Group = *Groups[*GroupIndex];
+
+		const int32 ItemIndex = [&Item, BeforeItemWithThisGuid, &Group, this]
+		{
+			if (BeforeItemWithThisGuid)
+			{
+				if (FItemLookup* LookupData = FindLookup(*BeforeItemWithThisGuid))
+				{
+					if (Group.Name == Groups[LookupData->GroupIndex]->Name)
+					{
+						// Set index for added item.
+						const int32 Index = LookupData->ItemIndex;
+
+						// Add item at the new index.
+						Group.Items.EmplaceAt(Index, Item);
+
+						// Increase index for the item we push back.
+						++LookupData->ItemIndex;
+
+						// Update all lookup data for items that come after the item we used to determine the insert location.
+						// This way we save the effort for finding the same lookup data again.
+						UpdateItemMapping(Group.Items, LookupData->ItemIndex + 1);
+
+						// Tell the outside where the new item was added.
+						return Index;
+					}
+				}
+			}
+
+			return Group.Items.Emplace(Item);
+		}();
+
+		ItemMapping.Emplace(Item->Guid, {*GroupIndex, ItemIndex});
+
+		Group.ListView->RequestListRefresh();
+
+		return true;
+	}
+
+	return false;
+}
+
+bool FImpl::MoveItem(const FGuid& Guid, const FName* GroupName, const FGuid* BeforeItemWithThisGuid)
+{
+	// Do not try to move an item before itself.
+	if (BeforeItemWithThisGuid && *BeforeItemWithThisGuid == Guid)
+	{
+		return false;
+	}
+
+	if (FItemLookup *const Lookup = FindLookup(Guid))
+	{
+		FGroup& GroupFrom = *Groups[Lookup->GroupIndex];
+
+		const int32 *const GroupToIndex = GroupName == nullptr ? &Lookup->GroupIndex : GroupMapping.Find(*GroupName);
+		if (GroupToIndex)
+		{
+			check(0 <= *GroupToIndex && *GroupToIndex < Groups.Num());
+			FGroup& GroupTo = *Groups[*GroupToIndex];
+
+			if (BeforeItemWithThisGuid)
+			{
+				if (FItemLookup *const LookupBefore = FindLookup(*BeforeItemWithThisGuid))
+				{
+					// Only move if BeforeItems is in the correct group and if item isn't already in the correct position.
+					if (*GroupToIndex == LookupBefore->GroupIndex &&
+						(Lookup->GroupIndex != LookupBefore->GroupIndex || Lookup->ItemIndex + 1 != LookupBefore->ItemIndex))
+					{
+						FItemType Item = MoveTemp(GroupFrom.Items[Lookup->ItemIndex]);
+
+						GroupFrom.Items.RemoveAt(Lookup->ItemIndex, EAllowShrinking::No);
+						UpdateItemMapping(GroupFrom.Items, Lookup->ItemIndex);
+
+						GroupTo.Items.EmplaceAt(LookupBefore->ItemIndex, MoveTemp(Item));
+						Lookup->GroupIndex = LookupBefore->GroupIndex;
+						Lookup->ItemIndex = LookupBefore->ItemIndex;
+						++LookupBefore->ItemIndex;
+						UpdateItemMapping(GroupTo.Items, LookupBefore->ItemIndex + 1);
+
+						return true;
+					}
+				}
+			}
+			else
+			{
+				// Only move if item is not already in the correct group.
+				if (*GroupToIndex != Lookup->GroupIndex)
+				{
+					FItemType Item = MoveTemp(GroupFrom.Items[Lookup->ItemIndex]);
+
+					GroupFrom.Items.RemoveAt(Lookup->ItemIndex, EAllowShrinking::No);
+					UpdateItemMapping(GroupFrom.Items, Lookup->ItemIndex);
+
+					Lookup->GroupIndex = *GroupToIndex;
+					Lookup->ItemIndex = GroupTo.Items.Emplace(MoveTemp(Item));
+
+					return true;
+				}
+			}
+		}
+	}
+
+	return false;
+}
+
+bool FImpl::RemoveItem(const FGuid& Guid)
+{
+	if (const FItemLookup *const Lookup = FindLookup(Guid))
+	{
+		FGroup& Group = *Groups[Lookup->GroupIndex];
+		Group.Items.RemoveAt(Lookup->ItemIndex, EAllowShrinking::No);
+
+		ItemMapping.Remove(Guid);
+		UpdateItemMapping(Group.Items, Lookup->ItemIndex);
+
+		Group.ListView->RequestListRefresh();
+
+		return true;
+	}
+
+	return false;
+}
+
+TSharedPtr<const FImageCatalogItemData> FImpl::GetItem(const FGuid& Guid) const
+{
+	if (const FItemLookup* Lookup = FindLookup(Guid))
+	{
+		return Groups[Lookup->GroupIndex]->Items[Lookup->ItemIndex];
+	}
+
+	return {};
+}
+
+TOptional<int32> FImpl::GetItemIndex(const FGuid& Guid) const
+{
+	if (const FItemLookup* Lookup = FindLookup(Guid))
+	{
+		return Lookup->ItemIndex;
+	}
+
+	return {};
+}
+
+TOptional<TTuple<FName, int32>> FImpl::GetItemGroupNameAndIndex(const FGuid& Guid) const
+{
+	if (const FItemLookup *const Lookup = FindLookup(Guid))
+	{
+		return {MakeTuple(Groups[Lookup->GroupIndex]->Name, Lookup->ItemIndex)};
+	}
+
+	return {};
+}
+
+TSharedPtr<const FImageCatalogItemData> FImpl::GetItemAt(const int32 Index, const FName* GroupName) const
+{
+	if (!GroupName)
+	{
+		GroupName = &DefaultGroupName;
+	}
+
+	if (const FGroup* const Group = FindGroup(*GroupName))
+	{
+		if (0 <= Index && Index < Group->Items.Num())
+		{
+			return {Group->Items[Index]};
+		}
+	}
+	return {};
+}
+
+TOptional<FGuid> FImpl::GetItemGuidAt(int32 Index, const FName* GroupName) const
+{
+	if (!GroupName)
+	{
+		GroupName = &DefaultGroupName;
+	}
+
+	if (const FGroup* const Group = FindGroup(*GroupName))
+	{
+		if (0 <= Index && Index < Group->Items.Num())
+		{
+			return {Group->Items[Index]->Guid};
+		}
+	}
+	return {};
+}
+
+bool FImpl::UpdateItem(const FGuid& Guid, const FSlateBrush* Thumbnail, const FText* Name, const FText* Info, const FText* ToolTip)
+{
+	if (const FItemLookup* Lookup = FindLookup(Guid))
+	{
+		const FItemType& ItemPtr = Groups[Lookup->GroupIndex]->Items[Lookup->ItemIndex];
+		check(ItemPtr.IsValid());
+
+		FImageCatalogItemData& Item = *ItemPtr;
+
+		if (Thumbnail) Item.Thumbnail = *Thumbnail;
+		if (Name) Item.Name = *Name;
+		if (Info) Item.Info = *Info;
+		if (ToolTip) Item.ToolTip = *ToolTip;
+
+		return true;
+	}
+
+	return false;
+}
+
+bool FImpl::SelectItem(const FGuid& Guid, const bool bSelected)
+{
+	if (const FItemLookup *const Lookup = FindLookup(Guid))
+	{
+		const FGroup& Group = *Groups[Lookup->GroupIndex];
+		const FItemType& Item = Group.Items[Lookup->ItemIndex];
+		Group.ListView->SetItemSelection(Item, bSelected);
+
+		return true;
+	}
+
+	return false;
+}
+
+bool FImpl::ClearSelection(const FName* GroupName)
+{
+	if (GroupName)
+	{
+		if (const FGroup *const Group = FindGroup(*GroupName))
+		{
+			Group->ListView->ClearSelection();
+			return true;
+		}
+
+		return false;
+	}
+
+	for (const TUniquePtr<FGroup>& Group : Groups)
+	{
+		Group->ListView->ClearSelection();
+	}
+
+	return true;
+}
+
+int32 FImpl::NumItems(const FName* GroupName) const
+{
+	if (GroupName)
+	{
+		if (const FGroup* Group = FindGroup(*GroupName))
+		{
+			return Group->Items.Num();
+		}
+
+		return 0;
+	}
+
+	return Algo::Accumulate(Groups, 0, [](int32 Num, const TUniquePtr<FGroup>& Group) { return Num + Group->Items.Num(); });
+}
+
+TOptional<FName> FImpl::GetItemGroupName(const FGuid& Guid) const
+{
+	if (const FItemLookup* Lookup = FindLookup(Guid))
+	{
+		return { Groups[Lookup->GroupIndex]->Name };
+	}
+
+	return {};
+}
+
+TSharedPtr<SWidget> FImpl::OnGroupContextMenuOpening(FName Name) const
+{
+	if (!OnGetGroupContextMenu.IsBound())
+	{
+		return SNullWidget::NullWidget;
+	}
+
+	return OnGetGroupContextMenu.Execute(Name);
+}
+
+TSharedPtr<SWidget> FImpl::OnItemsContextMenuOpening() const
+{
+	if (!OnGetItemsContextMenu.IsBound())
+	{
+		return SNullWidget::NullWidget;
+	}
+
+	TArray<FGuid> SelectedGuids;
+	for (const TUniquePtr<FGroup>& Group : Groups)
+	{
+		Algo::Transform(Group->ListView->GetSelectedItems(), SelectedGuids, [](const FItemType& Item) { return Item->Guid; });
+	}
+
+	if (SelectedGuids.IsEmpty())
+	{
+		return SNullWidget::NullWidget;
+	}
+
+	SortSelection(SelectedGuids);
+
+	return OnGetItemsContextMenu.Execute(SelectedGuids);
+}
+
+bool FImpl::AddGroup(FName Name, const FText& Heading, const int32 Index)
+{
+	check(0 <= Index && Index <= Groups.Num());
+
+	if (!GroupMapping.Contains(Name))
+	{
+		FGroup& Group = *Groups.EmplaceAt_GetRef(Index, MakeUnique<FGroup>(Name, Heading));
+		GroupMapping.Emplace(Name, Index);
+		UpdateGroupMapping(Index + 1);
+
+		const auto GetHeading = [&Group]
+		{
+			return Group.Heading;
+		};
+
+		const auto GetHeadingVisibility = [bShowEmptyGroups = bShowEmptyGroups, &Group]
+		{
+			return !Group.Heading.IsEmpty() && (bShowEmptyGroups || !Group.Items.IsEmpty()) ? EVisibility::Visible : EVisibility::Collapsed;
+		};
+
+		const auto GetHeadingHeight = [bShowEmptyGroups = bShowEmptyGroups, &Group]
+		{
+			return !Group.Heading.IsEmpty() && (bShowEmptyGroups || !Group.Items.IsEmpty()) ? 26.0f : 0.0f;
+		};
+
+		const auto GetVisibility = [bShowEmptyGroups = bShowEmptyGroups, &Group]
+		{
+			return bShowEmptyGroups || (Group.bIsExpanded && !Group.Items.IsEmpty()) ? EVisibility::Visible : EVisibility::Collapsed;
+		};
+
+		const auto GenerateItemRow = [](const FItemType& ItemData, const TSharedRef<STableViewBase>& OwnerTable)
+		{
+			static const FTableRowStyle TableRowStyle = []
+			{
+				FTableRowStyle Style = FAppStyle::Get().GetWidgetStyle<FTableRowStyle>("TableView.Row");
+				Style.SetOddRowBackgroundBrush(FSlateColorBrush(FStyleColors::Background));
+				Style.SetOddRowBackgroundHoveredBrush(FSlateColorBrush(FStyleColors::SelectHover));
+				Style.SetEvenRowBackgroundBrush(FSlateColorBrush(FStyleColors::Recessed));
+				Style.SetEvenRowBackgroundHoveredBrush(FSlateColorBrush(FStyleColors::SelectHover));
+				return Style;
+			}();
+
+			TSharedPtr<SImageCatalogItem> ItemWidget;
+			SAssignNew(ItemWidget, SImageCatalogItem, ItemData);
+
+			return SNew(STableRow<FItemType>, OwnerTable)
+				.Style(&TableRowStyle)
+				.ShowSelection(true)
+				[
+					ItemWidget.ToSharedRef()
+				];
+		};
+
+		const auto SelectionChanged = [bAllowSelectionAcrossGroups = bAllowSelectionAcrossGroups, &GroupName = Group.Name, &Groups = Groups, &OnItemSelected = OnItemSelected](const FItemType& Item, ESelectInfo::Type)
+		{
+			// Note that Item might be a nullptr since this callback is also executed when clearing the selection of the list.
+			if (Item.IsValid())
+			{
+				if (!bAllowSelectionAcrossGroups)
+				{
+					for (const TUniquePtr<FGroup>& Group : Groups)
+					{
+						if (Group->Name != GroupName)
+						{
+							if (Group->ListView)
+							{
+								Group->ListView->ClearSelection();
+							}
+						}
+					}
+				}
+
+				OnItemSelected.ExecuteIfBound(Item->Guid);
+			}
+		};
+
+		const auto OpenItemsContextMenu = [this]
+		{
+			return OnItemsContextMenuOpening();
+		};
+
+		TSharedPtr<SWidget> GroupHeader;
+
+		Layout->InsertSlot(Index)
+		[
+			SNew(SBorder)
+			.BorderImage(FAppStyle::Get().GetBrush("DetailsView.GridLine"))
+			.Padding(0.0f, 0.0f, 0.0f, 1.0f)
+			[
+				SNew(SVerticalBox)
+				+ SVerticalBox::Slot()
+				.AutoHeight()
+				[
+					SAssignNew(GroupHeader, SBox)
+					.MinDesiredHeight_Lambda(GetHeadingHeight)
+					.Visibility_Lambda(GetHeadingVisibility)
+					[
+						SNew(SBorder)
+						.BorderImage(FAppStyle::Get().GetBrush("DetailsView.CategoryTop"))
+						.Padding(0.0f)
+						[
+							SNew(SHorizontalBox)
+							+ SHorizontalBox::Slot()
+							.HAlign(HAlign_Left)
+							.VAlign(VAlign_Center)
+							.Padding(2.0f, 0.0f, 0.0f, 0.0f)
+							.AutoWidth()
+							[
+								SNew(SButton)
+								.ButtonStyle(FCoreStyle::Get(), "NoBorder")
+								.VAlign(VAlign_Center)
+								.HAlign(HAlign_Center)
+								.ClickMethod(EButtonClickMethod::MouseDown)
+								.OnClicked_Lambda([&bIsExpanded = Group.bIsExpanded]()
+								{
+									bIsExpanded = !bIsExpanded;
+									return FReply::Handled();
+								})
+								.ContentPadding(0.0f)
+								.IsFocusable(false)
+								[
+									SNew(SImage)
+									.Image_Lambda([&bIsExpanded = Group.bIsExpanded]()
+									{
+										return FAppStyle::Get().GetBrush(bIsExpanded ? "TreeArrow_Expanded" : "TreeArrow_Collapsed");
+									})
+									.ColorAndOpacity(FSlateColor::UseSubduedForeground())
+								]
+							]
+							+ SHorizontalBox::Slot()
+							.VAlign(VAlign_Center)
+							.FillWidth(1.0f)
+							.Padding(2.0f, 2.0f, 2.0f, 2.0f)
+							[
+								SNew(STextBlock)
+								.Text_Lambda(GetHeading)
+								.Font(FAppStyle::Get().GetFontStyle("DetailsView.CategoryFontStyle"))
+							]
+						]
+					]
+				]
+				+ SVerticalBox::Slot()
+				[
+					SAssignNew(Group.ListView, SListView<FItemType>)
+					.ListItemsSource(&Group.Items)
+					.Visibility_Lambda(GetVisibility)
+					.ScrollbarVisibility(EVisibility::Collapsed)
+					.OnGenerateRow_Lambda(GenerateItemRow)
+					.SelectionMode(SelectionMode)
+					.ClearSelectionOnClick(false)
+					.OnSelectionChanged_Lambda(SelectionChanged)
+					.OnContextMenuOpening_Lambda(OpenItemsContextMenu)
+				]
+			]
+		];
+
+		TWeakPtr<SWidget> WeakGroupHeader = GroupHeader;
+		GroupHeader->SetOnMouseButtonUp(FPointerEventHandler::CreateLambda([this, Name, WeakGroupHeader](const FGeometry& Geometry, const FPointerEvent& Event)
+		{
+			if (Event.GetEffectingButton() == EKeys::RightMouseButton)
+			{
+				if (const TSharedPtr<SWidget> PinnedHeader = WeakGroupHeader.Pin())
+				{
+					const TSharedPtr<SWidget> MenuContent = OnGroupContextMenuOpening(Name);
+					if (MenuContent.IsValid())
+					{
+						const FWidgetPath WidgetPath = Event.GetEventPath() != nullptr ? *Event.GetEventPath() : FWidgetPath();
+						FSlateApplication::Get().PushMenu(PinnedHeader->AsShared(), WidgetPath, MenuContent.ToSharedRef(), Event.GetScreenSpacePosition(),
+														  FPopupTransitionEffect(FPopupTransitionEffect::ContextMenu));
+					}
+
+					return FReply::Handled().ReleaseMouseCapture();
+				}
+			}
+
+			return FReply::Unhandled();
+		}));
+
+		return true;
+	}
+
+	return false;
+}
+
+const FGroup* FImpl::FindGroup(const FName Name) const
+{
+	if (const int32* Index = GroupMapping.Find(Name))
+	{
+		check(0 <= *Index && *Index < Groups.Num());
+		return Groups[*Index].Get();
+	}
+	return nullptr;
+}
+
+const FItemLookup* FImpl::FindLookup(const FGuid& Guid) const
+{
+	if (const FItemLookup* Lookup = ItemMapping.Find(Guid))
+	{
+		check(0 <= Lookup->GroupIndex && Lookup->GroupIndex < Groups.Num());
+		check(0 <= Lookup->ItemIndex && Lookup->ItemIndex < Groups[Lookup->GroupIndex]->Items.Num());
+
+		return Lookup;
+	}
+
+	return nullptr;
+}
+
+void FImpl::SortSelection(TArray<FGuid>& Selection) const
+{
+	TArray<int32> GroupOffsets;
+	GroupOffsets.Reserve(Groups.Num());
+	int32 Offset = 0;
+	for (const TUniquePtr<FGroup>& Group : Groups)
+	{
+		GroupOffsets.Add(Offset);
+		Offset += Group->Items.Num();
+	}
+
+	Algo::Sort(Selection, [&ItemMapping = ItemMapping, &GroupOffsets](const FGuid& A, const FGuid& B)
+	{
+		const FItemLookup *const LookupA = ItemMapping.Find(A);
+		const FItemLookup *const LookupB = ItemMapping.Find(B);
+
+		if (LookupA && !LookupB) return true;
+		if (!LookupA && LookupB) return false;
+		if (!LookupA && !LookupB) return false;
+
+		const int32 PositionA = GroupOffsets[LookupA->GroupIndex] + LookupA->ItemIndex;
+		const int32 PositionB = GroupOffsets[LookupB->GroupIndex] + LookupB->ItemIndex;
+
+		return PositionA < PositionB;
+	});
+}
+
+void FImpl::UpdateGroupMapping(const int32 StartingIndex)
+{
+	for (int32 GroupIndex = StartingIndex, NumGroups = Groups.Num(); GroupIndex < NumGroups; ++GroupIndex)
+	{
+		int32 *const GroupMappingIndex = GroupMapping.Find(Groups[GroupIndex]->Name);
+		check(GroupMappingIndex);
+		*GroupMappingIndex = GroupIndex;
+
+		const TArray<FItemType>& Items = Groups[GroupIndex]->Items;
+		for (int32 ItemIndex = 0, NumItems = Items.Num(); ItemIndex < NumItems; ++ItemIndex)
+		{
+			FItemLookup *const Lookup = ItemMapping.Find(Items[ItemIndex]->Guid);
+			check(Lookup);
+			Lookup->GroupIndex = GroupIndex;
+		}
+	}
+}
+
+void FImpl::UpdateItemMapping(TArray<FItemType>& Items, const int32 StartingIndex)
+{
+	for (int32 Index = StartingIndex, Num = Items.Num(); Index < Num; ++Index)
+	{
+		FItemLookup *const Lookup = ItemMapping.Find(Items[Index]->Guid);
+		check(Lookup);
+		Lookup->ItemIndex = Index;
+	}
+}
+
+void SImageCatalog::Construct(const FArguments& Args)
+{
+	TSharedPtr<SScrollBox> Layout;
+
+	ChildSlot
+	[
+		SAssignNew(Layout, SScrollBox)
+	];
+
+	Impl = MakePimpl<FImpl>(Args._DefaultGroupName, Args._DefaultGroupHeading, Args._SelectionMode, Args._bAllowSelectionAcrossGroups,
+	                        Args._bShowEmptyGroups, Args._OnItemSelected, Args._OnGetGroupContextMenu, Args._OnGetItemsContextMenu, Layout);
+}
+
+FName SImageCatalog::GetDefaultGroupName() const
+{
+	return Impl->GetDefaultGroupName();
+}
+
+bool SImageCatalog::AddGroup(const FName Name, const FText& Heading)
+{
+	return Impl->AddGroup(Name, Heading, nullptr);
+}
+
+bool SImageCatalog::AddGroup(const FName Name, const FText& Heading, const FName BeforeGroupWithThisName)
+{
+	return Impl->AddGroup(Name, Heading, &BeforeGroupWithThisName);
+}
+
+TOptional<TArray<FGuid>> SImageCatalog::RemoveGroup(const FName Name)
+{
+	return Impl->RemoveGroup(Name, nullptr);
+}
+
+TOptional<TArray<FGuid>> SImageCatalog::RemoveGroup(const FName Name, const FName GroupToMoveItemsInto)
+{
+	return Impl->RemoveGroup(Name, &GroupToMoveItemsInto);
+}
+
+int32 SImageCatalog::NumGroups() const
+{
+	return Impl->NumGroups();
+}
+
+TOptional<FName> SImageCatalog::GetGroupNameAt(int32 Index) const
+{
+	return Impl->GetGroupNameAt(Index);
+}
+
+bool SImageCatalog::SetGroupHeading(const FName Name, const FText& Heading)
+{
+	return Impl->SetGroupHeading(Name, Heading);
+}
+
+bool SImageCatalog::AddItem(const FItemType& Item)
+{
+	return Impl->AddItem(Item, nullptr, nullptr);
+}
+
+bool SImageCatalog::AddItem(const FItemType& Item, const FGuid& BeforeItemWithThisGuid)
+{
+	return Impl->AddItem(Item, nullptr, &BeforeItemWithThisGuid);
+}
+
+bool SImageCatalog::AddItem(const TSharedPtr<FImageCatalogItemData>& Item, const FName Group)
+{
+	return Impl->AddItem(Item, &Group, nullptr);
+}
+
+bool SImageCatalog::AddItem(const TSharedPtr<FImageCatalogItemData>& Item, const FName Group, const FGuid& BeforeItemWithThisGuid)
+{
+	return Impl->AddItem(Item, &Group, &BeforeItemWithThisGuid);
+}
+
+bool SImageCatalog::MoveItem(const FGuid& Guid, const FGuid& BeforeItemWithThisGuid)
+{
+	return Impl->MoveItem(Guid, nullptr, &BeforeItemWithThisGuid);
+}
+
+bool SImageCatalog::MoveItem(const FGuid& Guid, FName Group)
+{
+	return Impl->MoveItem(Guid, &Group, nullptr);
+}
+
+bool SImageCatalog::MoveItem(const FGuid& Guid, FName Group, const FGuid& BeforeItemWithThisGuid)
+{
+	return Impl->MoveItem(Guid, &Group, &BeforeItemWithThisGuid);
+}
+
+bool SImageCatalog::RemoveItem(const FGuid& Guid)
+{
+	return Impl->RemoveItem(Guid);
+}
+
+TSharedPtr<const FImageCatalogItemData> SImageCatalog::GetItem(const FGuid& Guid) const
+{
+	return Impl->GetItem(Guid);
+}
+
+TOptional<int32> SImageCatalog::GetItemIndex(const FGuid& Guid) const
+{
+	return Impl->GetItemIndex(Guid);
+}
+
+TOptional<TTuple<FName, int32>> SImageCatalog::GetItemGroupNameAndIndex(const FGuid& Guid) const
+{
+	return Impl->GetItemGroupNameAndIndex(Guid);
+}
+
+TSharedPtr<const FImageCatalogItemData> SImageCatalog::GetItemAt(const int32 Index) const
+{
+	return Impl->GetItemAt(Index, nullptr);
+}
+
+TSharedPtr<const FImageCatalogItemData> SImageCatalog::GetItemAt(const int32 Index, const FName Group) const
+{
+	return Impl->GetItemAt(Index, &Group);
+}
+
+TOptional<FGuid> SImageCatalog::GetItemGuidAt(const int32 Index) const
+{
+	return Impl->GetItemGuidAt(Index, nullptr);
+}
+
+TOptional<FGuid> SImageCatalog::GetItemGuidAt(const int32 Index, const FName Group) const 
+{
+	return Impl->GetItemGuidAt(Index, &Group);
+}
+
+int32 SImageCatalog::NumItems() const
+{
+	return Impl->NumItems(nullptr);
+}
+
+int32 SImageCatalog::NumItems(FName Group) const
+{
+	return Impl->NumItems(&Group);
+}
+
+TOptional<FName> SImageCatalog::GetItemGroupName(const FGuid& Guid) const
+{
+	return Impl->GetItemGroupName(Guid);
+}
+
+bool SImageCatalog::SelectItem(const FGuid& Guid)
+{
+	return Impl->SelectItem(Guid, true);
+}
+
+bool SImageCatalog::DeselectItem(const FGuid& Guid)
+{
+	return Impl->SelectItem(Guid, false);
+}
+
+void SImageCatalog::ClearSelection()
+{
+	Impl->ClearSelection(nullptr);
+}
+
+bool SImageCatalog::ClearSelection(FName Group)
+{
+	return Impl->ClearSelection(&Group);
+}
+
+bool SImageCatalog::UpdateItem(const FImageCatalogItemData& Item)
+{
+	return Impl->UpdateItem(Item.Guid, &Item.Thumbnail, &Item.Name, &Item.Info, &Item.ToolTip);
+}
+
+bool SImageCatalog::UpdateItemInfo(const FGuid& Guid, const FText& Info)
+{
+	return Impl->UpdateItem(Guid, nullptr, nullptr, &Info, nullptr);
+}
+
+bool SImageCatalog::UpdateItemName(const FGuid& Guid, const FText& Name)
+{
+	return Impl->UpdateItem(Guid, nullptr, &Name, nullptr, nullptr);
+}
+
+bool SImageCatalog::UpdateItemThumbnail(const FGuid& Guid, const FSlateBrush& Thumbnail)
+{
+	return Impl->UpdateItem(Guid, &Thumbnail, nullptr, nullptr, nullptr);
+}
+
+bool SImageCatalog::UpdateItemToolTip(const FGuid& Guid, const FText& ToolTip)
+{
+	return Impl->UpdateItem(Guid, nullptr, nullptr, nullptr, &ToolTip);
+}
+}
+
+#undef LOCTEXT_NAMESPACE
diff --git a/Engine/Plugins/Editor/ImageWidgets/Source/ImageWidgets/Private/SImageCatalogItem.cpp b/Engine/Plugins/Editor/ImageWidgets/Source/ImageWidgets/Private/SImageCatalogItem.cpp
new file mode 100644
--- /dev/null
+++ b/Engine/Plugins/Editor/ImageWidgets/Source/ImageWidgets/Private/SImageCatalogItem.cpp
@@ -0,0 +1,83 @@
+﻿// Copyright Epic Games, Inc. All Rights Reserved.
+
+#include "SImageCatalogItem.h"
+#include "SImageCatalog.h"
+
+#include "Widgets/SBoxPanel.h"
+#include "Widgets/Images/SImage.h"
+#include "Widgets/Text/STextBlock.h"
+
+#define LOCTEXT_NAMESPACE "SImageViewerCatalogItem"
+
+namespace UE::ImageWidgets
+{
+void SImageCatalogItem::Construct(const FArguments& InArgs, const TSharedPtr<FImageCatalogItemData>& InItemData)
+{
+	ItemData = InItemData;
+
+	ChildSlot
+	[
+		SNew(SHorizontalBox)
+		+ SHorizontalBox::Slot()
+			.Padding(FMargin(2.0f, 2.0f))
+			.AutoWidth()
+		[
+			SNew(SImage)
+				.DesiredSizeOverride(FVector2D(32.0f, 32.0f))
+				.Image(this, &SImageCatalogItem::GetItemThumbnail)
+				.ToolTipText(this, &SImageCatalogItem::GetItemToolTip)
+		]
+		+ SHorizontalBox::Slot()
+			.FillWidth(1.0f)
+			.VAlign(VAlign_Center)
+			.HAlign(HAlign_Fill)
+		[
+			SNew(SVerticalBox)
+			+ SVerticalBox::Slot()
+				.AutoHeight()
+				.HAlign(HAlign_Fill)
+				.Padding(FMargin(1.0f, 1.0f))
+			[
+				SNew(STextBlock)
+					.Text(this, &SImageCatalogItem::GetItemName)
+			]
+			+ SVerticalBox::Slot()
+				.AutoHeight()
+				.HAlign(HAlign_Fill)
+				.Padding(FMargin(1.0f, 1.0f))
+			[
+			SNew(SHorizontalBox)
+				+ SHorizontalBox::Slot()
+					.FillWidth(1.0f)
+				[
+					SNew(STextBlock)
+						.Text(this, &SImageCatalogItem::GetItemInfo)
+						.TextStyle(&FAppStyle::Get(), "SmallText.Subdued")
+				]
+			]
+		]
+	];
+}
+
+FText SImageCatalogItem::GetItemInfo() const
+{
+	return ItemData->Info;
+}
+
+FText SImageCatalogItem::GetItemName() const
+{
+	return ItemData->Name;
+}
+
+const FSlateBrush* SImageCatalogItem::GetItemThumbnail() const
+{
+	return &ItemData->Thumbnail;
+}
+
+FText SImageCatalogItem::GetItemToolTip() const
+{
+	return ItemData->ToolTip;
+}
+}
+
+#undef LOCTEXT_NAMESPACE
diff --git a/Engine/Plugins/Editor/ImageWidgets/Source/ImageWidgets/Private/SImageCatalogItem.h b/Engine/Plugins/Editor/ImageWidgets/Source/ImageWidgets/Private/SImageCatalogItem.h
new file mode 100644
--- /dev/null
+++ b/Engine/Plugins/Editor/ImageWidgets/Source/ImageWidgets/Private/SImageCatalogItem.h
@@ -0,0 +1,31 @@
+﻿// Copyright Epic Games, Inc. All Rights Reserved.
+
+#pragma once
+
+#include "CoreMinimal.h"
+#include "SImageCatalog.h"
+
+namespace UE::ImageWidgets
+{
+	struct FImageCatalogItemData;
+
+	/**
+	 * Widget for a single item row in the catalog.
+	 */
+	class SImageCatalogItem : public SCompoundWidget
+	{
+	public:
+		SLATE_BEGIN_ARGS(SImageCatalogItem) {}
+		SLATE_END_ARGS()
+
+		void Construct(const FArguments& InArgs, const TSharedPtr<FImageCatalogItemData>& InItemData);
+
+	private:
+		FText GetItemInfo() const;
+		FText GetItemName() const;
+		const FSlateBrush* GetItemThumbnail() const;
+		FText GetItemToolTip() const;
+
+		TSharedPtr<FImageCatalogItemData> ItemData;
+	};
+}
diff --git a/Engine/Plugins/Editor/ImageWidgets/Source/ImageWidgets/Private/SImageViewport.cpp b/Engine/Plugins/Editor/ImageWidgets/Source/ImageWidgets/Private/SImageViewport.cpp
new file mode 100644
--- /dev/null
+++ b/Engine/Plugins/Editor/ImageWidgets/Source/ImageWidgets/Private/SImageViewport.cpp
@@ -0,0 +1,512 @@
+﻿// Copyright Epic Games, Inc. All Rights Reserved.
+
+#include "SImageViewport.h"
+
+#include "CanvasTypes.h"
+#include "ImageABComparison.h"
+#include "ImageViewportClient.h"
+#include "ImageViewportController.h"
+#include "ImageWidgetsCommands.h"
+#include "ImageWidgetsStyle.h"
+#include "SImageViewportToolbar.h"
+#include "Framework/Application/SlateApplication.h"
+#include "Widgets/Text/SRichTextBlock.h"
+#include "Widgets/Text/STextBlock.h"
+
+#define LOCTEXT_NAMESPACE "ImageViewport"
+
+namespace UE::ImageWidgets
+{
+	class FStatusBarExtension
+	{
+	public:
+		virtual ~FStatusBarExtension() = default;
+
+		FName Hook;
+		EExtensionHook::Position HookPosition = EExtensionHook::Before;
+		TSharedPtr<FUICommandList> CommandList;
+		SImageViewport::FStatusBarExtender::FDelegate Delegate;
+	};
+
+	void SImageViewport::FStatusBarExtender::AddExtension(FName ExtensionHook, EExtensionHook::Position HookPosition,
+	                                                      const TSharedPtr<FUICommandList>& Commands, const FDelegate& Delegate)
+	{
+		const TPimplPtr<FStatusBarExtension>& Extension = Extensions.Emplace_GetRef(MakePimpl<FStatusBarExtension>());
+		Extension->Hook = ExtensionHook;
+		Extension->HookPosition = HookPosition;
+		Extension->CommandList = Commands;
+		Extension->Delegate = Delegate;
+	}
+
+	void SImageViewport::FStatusBarExtender::Apply(FName ExtensionHook, EExtensionHook::Position HookPosition, SHorizontalBox& HorizontalBox) const
+	{
+		for (const TPimplPtr<FStatusBarExtension>& Extension : Extensions)
+		{
+			if (Extension->Hook == ExtensionHook && Extension->HookPosition == HookPosition)
+			{
+				Extension->Delegate.ExecuteIfBound(HorizontalBox);
+			}
+		}
+	}
+
+	SImageViewport::SImageViewport()
+		: ABComparison(MakePimpl<FImageABComparison>(
+			FImageABComparison::FImageIsValid::CreateLambda([this](const FGuid& Guid) { return ImageViewer->IsValidImage(Guid); }),
+			FImageABComparison::FGetCurrentImageGuid::CreateLambda([this] { return ImageViewer->GetCurrentImageInfo().Guid; }),
+			FImageABComparison::FGetImageName::CreateLambda([this](const FGuid& Guid) { return ImageViewer->GetImageName(Guid); })))
+	{
+	}
+
+	SImageViewport::~SImageViewport()
+	{
+		FImageWidgetsCommands::Unregister();
+	}
+
+	void SImageViewport::BindCommands()
+	{
+		const FImageWidgetsCommands& Commands = FImageWidgetsCommands::Get();
+
+		auto ToggleOverlay = [this]
+		{
+			ViewportOverlay->SetVisibility(ViewportOverlay->GetVisibility() == EVisibility::Visible ? EVisibility::Hidden : EVisibility::Visible);
+		};
+
+		CommandList->MapAction(
+			Commands.ToggleOverlay,
+			FExecuteAction::CreateLambda(ToggleOverlay)
+		);
+
+		// ZOOM
+
+		auto MapZoom = [this](const TSharedPtr<FUICommandInfo>& Command, double ZoomFactor)
+		{
+			CommandList->MapAction(
+				Command, FExecuteAction::CreateSP(ImageViewportClient.ToSharedRef(), &FImageViewportClient::SetZoom,
+				                                  FImageViewportController::EZoomMode::Custom, ZoomFactor));
+		};
+
+		MapZoom(Commands.Zoom12, 0.125);
+		MapZoom(Commands.Zoom25, 0.25);
+		MapZoom(Commands.Zoom50, 0.5);
+		MapZoom(Commands.Zoom100, 1.0);
+		MapZoom(Commands.Zoom200, 2.0);
+		MapZoom(Commands.Zoom400, 4.0);
+		MapZoom(Commands.Zoom800, 8.0);
+
+		auto MapZoomFitFill = [this](const TSharedPtr<FUICommandInfo>& Command, FImageViewportController::EZoomMode ZoomMode)
+		{
+			CommandList->MapAction(
+				Command,
+				FExecuteAction::CreateSP(ImageViewportClient.ToSharedRef(), &FImageViewportClient::SetZoom, ZoomMode, 0.0),
+				FCanExecuteAction(),
+				FIsActionChecked::CreateLambda([this, ZoomMode] { return ImageViewportClient->GetZoom().Mode == ZoomMode; }));
+		};
+
+		MapZoomFitFill(Commands.ZoomFit, FImageViewportController::EZoomMode::Fit);
+		MapZoomFitFill(Commands.ZoomFill, FImageViewportController::EZoomMode::Fill);
+
+		// MIPS
+
+		CommandList->MapAction(
+			Commands.MipMinus,
+			FExecuteAction::CreateLambda([this]()
+			{
+				const int32 CurrentMip = ImageViewportClient->GetMipLevel();
+
+				int32 NumMips = 0;
+				if (ImageViewer.IsValid())
+				{
+					const IImageViewer::FImageInfo ImageInfo = ImageViewer->GetCurrentImageInfo();
+					if (ImageInfo.bIsValid)
+					{
+						NumMips = ImageInfo.NumMips;
+					}
+				}
+
+				if (NumMips > 0 && CurrentMip < NumMips - 1)
+				{
+					ImageViewportClient->SetMipLevel(CurrentMip + 1);
+				}
+			}));
+		CommandList->MapAction(
+			Commands.MipPlus,
+			FExecuteAction::CreateLambda([this]()
+			{
+				const int32 CurrentMip = ImageViewportClient->GetMipLevel();
+				if (CurrentMip > INDEX_NONE)
+				{
+					ImageViewportClient->SetMipLevel(CurrentMip - 1);
+				}
+			}));
+	}
+
+	void SImageViewport::Construct(const FArguments& InArgs, const TSharedRef<IImageViewer>& InImageViewer)
+	{
+		FImageWidgetsCommands::Register();
+
+		ImageViewer = InImageViewer;
+		check(ImageViewer.IsValid());
+
+		ToolbarExtender = InArgs._ToolbarExtender;
+		ToolbarExtenderStyle = InArgs._ToolbarExtenderStyle;
+		ToolbarExtenderStyleName = InArgs._ToolbarExtenderStyleName;
+
+		StatusBarExtender = InArgs._StatusBarExtender;
+
+		DrawSettings = InArgs._DrawSettings;
+		OverlaySettings = InArgs._OverlaySettings;
+		bABComparisonEnabled = InArgs._bABComparisonEnabled;
+
+		auto GetImageSize = [this]
+		{
+			if (ImageViewer.IsValid())
+			{
+				const IImageViewer::FImageInfo ImageInfo = ImageViewer->GetCurrentImageInfo();
+				return ImageInfo.Size;
+			}
+			return FIntPoint::ZeroValue;
+		};
+
+		auto DrawImage = [this](FViewport* Viewport, FCanvas* Canvas, const IImageViewer::FDrawProperties& DrawProperties)
+		{
+			ImageViewer->DrawCurrentImage(Viewport, Canvas, DrawProperties);
+		};
+
+		auto GetDrawSettings = [this]
+		{
+			return DrawSettings.Get({});
+		};
+
+		ImageViewportClient = MakeShareable(new FImageViewportClient(
+			StaticCastWeakPtr<SEditorViewport>(AsWeak()),
+			FGetImageSize::CreateLambda(GetImageSize),
+			FDrawImage::CreateLambda(DrawImage),
+			FGetDrawSettings::CreateLambda(GetDrawSettings),
+			FGetLayoutAndDPIScaleFactor::CreateRaw(this, &SImageViewport::GetLayoutAndDPIScaleFactor),
+			FGetLayoutScaleFactor::CreateRaw(this, &SImageViewport::GetLayoutScaleFactor),
+			ABComparison.Get(),
+			InArgs._ControllerSettings)
+			);
+
+		SEditorViewport::Construct(SEditorViewport::FArguments());
+	}
+
+	void SImageViewport::ResetController(FIntPoint ImageSize) const
+	{
+		ImageViewportClient->ResetController(ImageSize);
+	}
+
+	void SImageViewport::ResetMip() const
+	{
+		const int32 SelectedMip = ImageViewportClient->GetMipLevel();
+		if (SelectedMip != INDEX_NONE)
+		{
+			if (ImageViewer.IsValid())
+			{
+				const IImageViewer::FImageInfo ImageInfo = ImageViewer->GetCurrentImageInfo();
+				if (ImageInfo.bIsValid)
+				{
+					if (ImageInfo.NumMips <= SelectedMip)
+					{
+						ImageViewportClient->SetMipLevel(ImageInfo.NumMips - 1);
+					}
+				}
+			}
+		}
+	}
+
+	void SImageViewport::ResetZoom(FIntPoint ImageSize) const
+	{
+		ImageViewportClient->ResetZoom(ImageSize);
+	}
+
+	TSharedPtr<SViewportToolBar> SImageViewport::GetParentToolbar() const
+	{
+		return ImageViewportToolbar;
+	}
+
+	SImageViewport::FPixelCoordinatesUnderCursorResult SImageViewport::GetPixelCoordinatesUnderCursor() const
+	{
+		const auto [PixelCoordsValid, PixelCoords] = ImageViewportClient->GetPixelCoordinatesUnderCursor();
+		return {PixelCoordsValid, PixelCoords};
+	}
+
+	void SImageViewport::RequestRedraw() const
+	{
+		const TSharedPtr<FEditorViewportClient> ViewportClient = GetViewportClient();
+		if (ViewportClient.IsValid())
+		{
+			ViewportClient->RedrawRequested(ViewportClient->Viewport);
+		}
+	}
+
+	TSharedRef<FEditorViewportClient> SImageViewport::MakeEditorViewportClient()
+	{
+		if (!Client.IsValid())
+		{
+			Client = ImageViewportClient;
+		}
+
+		return Client.ToSharedRef();
+	}
+
+	TSharedRef<SWidget> SImageViewport::MakeViewportToolbarOverlay()
+	{
+		auto HasImage = [this]
+		{
+			if (ImageViewer.IsValid())
+			{
+				const IImageViewer::FImageInfo ImageInfo = ImageViewer->GetCurrentImageInfo();
+				return ImageInfo.bIsValid;
+			}
+			return false;
+		};
+
+		auto GetNumMips = [this]
+		{
+			if (ImageViewer.IsValid())
+			{
+				const IImageViewer::FImageInfo ImageInfo = ImageViewer->GetCurrentImageInfo();
+				if (ImageInfo.bIsValid)
+				{
+					return ImageInfo.NumMips;
+				}
+			}
+
+			return 0;
+		};
+
+		auto GetImageGuid = [this]
+		{
+			return ImageViewer.IsValid() ? ImageViewer->GetCurrentImageInfo().Guid : FGuid();
+		};
+
+		auto GetOverlaySettings = [this]
+		{
+			return OverlaySettings.Get({});
+		};
+
+		SAssignNew(ImageViewportToolbar, SImageViewportToolbar, StaticCastSharedPtr<FImageViewportClient>(Client), CommandList,
+		           SImageViewportToolbar::FConstructParameters{
+			           SImageViewportToolbar::FHasImage::CreateLambda(HasImage),
+			           SImageViewportToolbar::FNumMips::CreateLambda(GetNumMips),
+			           SImageViewportToolbar::FImageGuid::CreateLambda(GetImageGuid),
+			           FGetLayoutAndDPIScaleFactor::CreateRaw(this, &SImageViewport::GetLayoutAndDPIScaleFactor),
+			           SImageViewportToolbar::FGetOverlaySettings::CreateLambda(GetOverlaySettings),
+			           bABComparisonEnabled ? ABComparison.Get() : nullptr,
+			           ToolbarExtender, ToolbarExtenderStyle, ToolbarExtenderStyleName
+		           });
+
+		ToolbarExtender.Reset();
+		ToolbarExtenderStyle = nullptr;
+		ToolbarExtenderStyleName = FName();
+
+		return ImageViewportToolbar.ToSharedRef();
+	}
+
+	void SImageViewport::PopulateViewportOverlays(TSharedRef<SOverlay> Overlay)
+	{
+		Overlay->AddSlot()
+			.VAlign(VAlign_Top)
+			[
+				MakeViewportToolbarOverlay()
+			];
+
+		Overlay->AddSlot()
+		       .VAlign(VAlign_Bottom)
+		[
+			MakeStatusBar(Overlay)
+		];
+	}
+
+	void SImageViewport::OnArrangeChildren(const FGeometry& AllottedGeometry, FArrangedChildren& ArrangedChildren) const
+	{
+		SEditorViewport::OnArrangeChildren(AllottedGeometry, ArrangedChildren);
+
+		AllottedGeometryScale = AllottedGeometry.Scale;
+	}
+
+	float SImageViewport::GetDPIScaleFactor() const
+	{
+		if (const TSharedPtr<SWindow> TopLevelWindow = FSlateApplication::Get().FindWidgetWindow(SharedThis(this)); TopLevelWindow.IsValid())
+		{
+			return TopLevelWindow->GetDPIScaleFactor();
+		}
+		return 1.0f;
+	}
+
+	float SImageViewport::GetLayoutScaleFactor() const
+	{
+		const float DPI = GetDPIScaleFactor();
+		if (DPI > SMALL_NUMBER)
+		{
+			return FMath::Max(AllottedGeometryScale / DPI, SMALL_NUMBER);
+		}
+		return 1.0f;
+	}
+
+	float SImageViewport::GetLayoutAndDPIScaleFactor() const
+	{
+		return FMath::Max(AllottedGeometryScale, SMALL_NUMBER);
+	}
+
+	FText SImageViewport::GetPickerLabel() const
+	{
+		if (!ImageViewer.IsValid())
+		{
+			return {};
+		}
+
+		const IImageViewer::FImageInfo ImageInfo = ImageViewer->GetCurrentImageInfo();
+		if (ImageInfo.bIsValid)
+		{
+			const auto [PixelCoordsValid, PixelCoords] = ImageViewportClient->GetPixelCoordinatesUnderCursor();
+			if (PixelCoordsValid)
+			{
+				const FIntPoint PixelCoordsInt(PixelCoords.X, PixelCoords.Y);
+				if (0 <= PixelCoordsInt.X && 0 <= PixelCoordsInt.Y && PixelCoordsInt.X < ImageInfo.Size.X && PixelCoordsInt.Y < ImageInfo.Size.Y)
+				{
+					const int32 MipIndex = FMath::Max(0, ImageViewportClient->GetMipLevel());
+
+					const FIntPoint MipPixelCoords = [&PixelCoordsInt, MipIndex]
+					{
+						FIntPoint Coords = PixelCoordsInt;
+						if (MipIndex > 0)
+						{
+							Coords.X /= 1 << MipIndex;
+							Coords.Y /= 1 << MipIndex;
+						}
+						return Coords;
+					}();
+
+					const TOptional<TVariant<FColor, FLinearColor>> PixelColor = ImageViewer->GetCurrentImagePixelColor(MipPixelCoords, MipIndex);
+
+					auto CoordsAndColorText = [&MipPixelCoords](const auto* Color, const FNumberFormattingOptions& Formatting)
+					{
+						return FText::Format(
+							LOCTEXT("CoordinatesWithColor", "x={0} y={1}   <RichTextBlock.Red>{2}</> <RichTextBlock.Green>{3}</> <RichTextBlock.Blue>{4}</> {5}"),
+							FText::AsNumber(MipPixelCoords.X), FText::AsNumber(MipPixelCoords.Y),
+							FText::AsNumber(Color->R, &Formatting), FText::AsNumber(Color->G, &Formatting),
+							FText::AsNumber(Color->B, &Formatting), FText::AsNumber(Color->A, &Formatting));
+					};
+
+					FNumberFormattingOptions FormattingByte;
+					FormattingByte.SetMinimumIntegralDigits(3);
+
+					FNumberFormattingOptions FormattingFloat;
+					FormattingFloat.SetMinimumFractionalDigits(3);
+					FormattingFloat.SetMaximumFractionalDigits(3);
+
+					if (PixelColor.IsSet())
+					{
+						const TVariant<FColor, FLinearColor>& PixelColorValue = PixelColor.GetValue();
+						if (const FColor* Color = PixelColorValue.TryGet<FColor>())
+						{
+							return CoordsAndColorText(Color, FormattingByte);
+						}
+						if (const FLinearColor* ColorLinear = PixelColorValue.TryGet<FLinearColor>())
+						{
+							return CoordsAndColorText(ColorLinear, FormattingFloat);
+						}
+					}
+
+					return FText::Format(
+						LOCTEXT("CoordinatesOnly", "x={0} y={1}"),
+						FText::AsNumber(MipPixelCoords.X), FText::AsNumber(MipPixelCoords.Y));
+				}
+			}
+		}
+		return {};
+	}
+
+	FText SImageViewport::GetResolutionLabel() const
+	{
+		if (ImageViewer.IsValid())
+		{
+			const IImageViewer::FImageInfo ImageInfo = ImageViewer->GetCurrentImageInfo();
+			if (ImageInfo.bIsValid)
+			{
+				FIntPoint Size = ImageInfo.Size;
+
+				const int32 Mip = ImageViewportClient->GetMipLevel();
+				if (Mip > 0)
+				{
+					Size.X = FMath::Max(1, Size.X / (1 << Mip));
+					Size.Y = FMath::Max(1, Size.Y / (1 << Mip));
+				}
+
+				return FText::Format(LOCTEXT("Resolution", "{0} \u00D7 {1}"), FText::AsNumber(Size.X), FText::AsNumber(Size.Y));
+			}
+		}
+		return {};
+	}
+
+	TSharedRef<SWidget> SImageViewport::MakeStatusBar(TSharedRef<SOverlay> Overlay)
+	{
+		const FMargin SlotPadding(6.0f, 2.0f);
+
+		const TSharedPtr<SHorizontalBox> HorizontalBox = SNew(SHorizontalBox);
+
+		auto ApplyHook = [this, &HorizontalBox](FName ExtensionHook, EExtensionHook::Position HookPosition)
+		{
+			if (ExtensionHook != NAME_None && StatusBarExtender.IsValid())
+			{
+				StatusBarExtender->Apply(ExtensionHook, HookPosition, *HorizontalBox);
+			}
+		};
+
+		ApplyHook(FName("StatusBarLeft"), EExtensionHook::Before);
+		{
+			HorizontalBox->AddSlot()
+			             .Padding(SlotPadding)
+			             .AutoWidth()
+			             .HAlign(HAlign_Left)
+			[
+				SNew(STextBlock)
+					.Visibility_Lambda([&OverlaySettings = OverlaySettings]
+						{
+							return OverlaySettings.IsSet() && OverlaySettings.Get().bDisableStatusBarLeft ? EVisibility::Collapsed : EVisibility::Visible;
+						})
+					.Text(this, &SImageViewport::GetResolutionLabel)
+			];
+		}
+		ApplyHook(FName("StatusBarLeft"), EExtensionHook::After);
+
+		ApplyHook(FName("StatusBarCenter"), EExtensionHook::Before);
+		{
+			HorizontalBox->AddSlot()
+			             .HAlign(HAlign_Center)
+			[
+				SNullWidget::NullWidget
+
+				// This is deliberately left empty.
+				// It allows for adding toolbar extensions though.
+			];
+		}
+		ApplyHook(FName("StatusBarCenter"), EExtensionHook::After);
+
+		ApplyHook(FName("StatusBarRight"), EExtensionHook::Before);
+		{
+			HorizontalBox->AddSlot()
+			             .Padding(SlotPadding)
+			             .AutoWidth()
+			             .HAlign(HAlign_Right)
+			[
+				SNew(SRichTextBlock)
+					.Visibility_Lambda([&OverlaySettings = OverlaySettings]
+						{
+							return OverlaySettings.IsSet() && OverlaySettings.Get().bDisableStatusBarRight ? EVisibility::Collapsed : EVisibility::Visible;
+						})
+					.Text(this, &SImageViewport::GetPickerLabel)
+					.DecoratorStyleSet(&FImageWidgetsStyle::Get())
+			];
+		}
+		ApplyHook(FName("StatusBarRight"), EExtensionHook::After);
+
+		StatusBarExtender.Reset();
+
+		return HorizontalBox.ToSharedRef();
+	}
+}
+
+#undef LOCTEXT_NAMESPACE
diff --git a/Engine/Plugins/Editor/ImageWidgets/Source/ImageWidgets/Private/SImageViewportToolbar.cpp b/Engine/Plugins/Editor/ImageWidgets/Source/ImageWidgets/Private/SImageViewportToolbar.cpp
new file mode 100644
--- /dev/null
+++ b/Engine/Plugins/Editor/ImageWidgets/Source/ImageWidgets/Private/SImageViewportToolbar.cpp
@@ -0,0 +1,316 @@
+﻿// Copyright Epic Games, Inc. All Rights Reserved.
+
+#include "SImageViewportToolbar.h"
+
+#include "ImageABComparison.h"
+#include "ImageViewportClient.h"
+#include "ImageWidgetsCommands.h"
+#include "SEditorViewportToolBarMenu.h"
+#include "Framework/MultiBox/MultiBoxBuilder.h"
+#include "Widgets/SBoxPanel.h"
+#include "Widgets/Input/SCheckBox.h"
+#include "Widgets/Layout/SBorder.h"
+
+#define LOCTEXT_NAMESPACE "ImageViewportToolbar"
+
+namespace UE::ImageWidgets
+{
+	namespace ImageViewportToolbar_Local
+	{
+		FText Auto = LOCTEXT("Auto", "Auto");
+	}
+
+	void SImageViewportToolbar::Construct(const FArguments& InArgs, const TSharedPtr<FImageViewportClient>& InViewportClient,
+	                                      const TSharedPtr<FUICommandList>& InCommandList, FConstructParameters Parameters)
+	{
+		ViewportClient = InViewportClient;
+		CommandList = InCommandList;
+		check(ViewportClient.IsValid());
+		check(CommandList.IsValid());
+
+		HasImage = MoveTemp(Parameters.HasImage);
+		NumMips = MoveTemp(Parameters.NumMips);
+		ImageGuid = MoveTemp(Parameters.ImageGuid);
+		GetLayoutAndDPIScaleFactor = MoveTemp(Parameters.GetLayoutAndDPIScaleFactor);
+		GetOverlaySettings = MoveTemp(Parameters.GetOverlaySettings);
+		check(HasImage.IsBound());
+		check(NumMips.IsBound());
+		check(ImageGuid.IsBound());
+		check(GetLayoutAndDPIScaleFactor.IsBound());
+		check(GetOverlaySettings.IsBound());
+
+		ABComparison = Parameters.ABComparison;
+
+		ToolbarExtenderStyle = Parameters.ToolbarExtenderStyle;
+		ToolbarExtenderStyleName = Parameters.ToolbarExtenderStyleName;
+
+		ChildSlot
+		[
+			SNew(SBorder)
+			.BorderImage(FAppStyle::GetBrush("EditorViewportToolBar.Background"))
+			[
+				SNew(SHorizontalBox)
+				+ SHorizontalBox::Slot()
+				  .AutoWidth()
+				  .HAlign(HAlign_Left)
+				[
+					MakeLeftToolbar(Parameters.ToolbarExtender)
+				]
+				+ SHorizontalBox::Slot()
+				  .FillWidth(1.0f)
+				  .HAlign(HAlign_Center)
+				[
+					MakeCenterToolbar(Parameters.ToolbarExtender, ABComparison != nullptr)
+				]
+				+ SHorizontalBox::Slot()
+				  .AutoWidth()
+				  .HAlign(HAlign_Right)
+				[
+					MakeRightToolbar(Parameters.ToolbarExtender)
+				]
+			]
+		];
+
+		SViewportToolBar::Construct(SViewportToolBar::FArguments());
+	}
+
+	FSlimHorizontalToolBarBuilder SImageViewportToolbar::GetToolbarBuilder(const TSharedPtr<FExtender>& Extender) const
+	{
+		FSlimHorizontalToolBarBuilder ToolbarBuilder(CommandList, FMultiBoxCustomization::None, Extender, true);
+
+		const ISlateStyle* Style = ToolbarExtenderStyle ? ToolbarExtenderStyle : &FAppStyle::Get();
+		const FName StyleName = ToolbarExtenderStyle && !ToolbarExtenderStyleName.IsNone() ? ToolbarExtenderStyleName : FName("EditorViewportToolBar");
+		ToolbarBuilder.SetStyle(Style, StyleName);
+
+		ToolbarBuilder.SetIsFocusable(false);
+		ToolbarBuilder.SetLabelVisibility(EVisibility::Collapsed);
+
+		return ToolbarBuilder;
+	}
+
+	TSharedRef<SWidget> SImageViewportToolbar::MakeLeftToolbar(const TSharedPtr<FExtender>& Extender)
+	{
+		FSlimHorizontalToolBarBuilder ToolbarBuilder = GetToolbarBuilder(Extender);
+
+		ToolbarBuilder.BeginSection("ToolbarLeft");
+		{
+			ToolbarBuilder.BeginBlockGroup();
+
+			ToolbarBuilder.AddWidget(
+				SNew(SEditorViewportToolbarMenu)
+				.ParentToolBar(SharedThis(this))
+				.Cursor(EMouseCursor::Default)
+				.Label(this, &SImageViewportToolbar::GetZoomMenuLabel)
+				.OnGetMenuContent(this, &SImageViewportToolbar::MakeZoomMenu)
+				.IsEnabled_Lambda([this] { return HasImage.Execute(); })
+				.Visibility(this, &SImageViewportToolbar::GetZoomMenuVisibility)
+			);
+
+			ToolbarBuilder.AddSeparator();
+
+			ToolbarBuilder.AddWidget(
+				SNew(SEditorViewportToolbarMenu)
+				.ParentToolBar(SharedThis(this))
+				.Cursor(EMouseCursor::Default)
+				.Label(this, &SImageViewportToolbar::GetMipMenuLabel)
+				.OnGetMenuContent(this, &SImageViewportToolbar::MakeMipMenu)
+				.IsEnabled_Lambda([this] { return HasImage.Execute(); })
+				.Visibility(this, &SImageViewportToolbar::GetMipMenuVisibility)
+			);
+
+			ToolbarBuilder.EndBlockGroup();
+		}
+		ToolbarBuilder.EndSection();
+
+		return ToolbarBuilder.MakeWidget();
+	}
+
+	TSharedRef<SWidget> SImageViewportToolbar::MakeCenterToolbar(const TSharedPtr<FExtender>& Extender, bool bEnableABComparison)
+	{
+		FSlimHorizontalToolBarBuilder ToolbarBuilder = GetToolbarBuilder(Extender);
+
+		ToolbarBuilder.BeginSection("ToolbarCenter");
+		{
+			ToolbarBuilder.AddSeparator();
+
+			if (bEnableABComparison && ABComparison)
+			{
+				auto GetTextButton = [this](const FString& Label, const FCheckBoxStyle* ButtonStyle, FImageABComparison::EAorB AorB)
+				{
+					return SNew(SCheckBox)
+						.Style(ButtonStyle)
+						.Visibility(this, &SImageViewportToolbar::GetABVisibility)
+						.IsEnabled_Lambda([this, AorB]
+							{
+								return ABComparison->CanSetABComparison(AorB);
+							})
+						.IsChecked_Lambda([this, AorB]
+							{
+								return ABComparison->ABComparisonIsSet(AorB) ? ECheckBoxState::Checked : ECheckBoxState::Unchecked;
+							})
+						.OnCheckStateChanged_Lambda([this, AorB](const ECheckBoxState State)
+						    {
+								ABComparison->SetABComparison(AorB, State != ECheckBoxState::Checked ? FGuid() : ImageGuid.Execute());
+						    })
+						[
+							SNew(STextBlock)
+								.Font(FAppStyle::GetFontStyle("EditorViewportToolBar.Font"))
+								.Text(FText::FromString(Label))
+								.ToolTipText_Lambda([&ABComparison = ABComparison, AorB]
+									{
+										return ABComparison->ABComparisonIsSet(AorB) ? ABComparison->GetName(AorB) : FText();
+									})
+								.Margin(FMargin(2.0f, 0.0f))
+						];
+				};
+
+				const FCheckBoxStyle* ButtonStyleStart = &FAppStyle::Get().GetWidgetStyle<FCheckBoxStyle>("EditorViewportToolBar.ToggleButton.Start");
+				const FCheckBoxStyle* ButtonStyleEnd = &FAppStyle::Get().GetWidgetStyle<FCheckBoxStyle>("EditorViewportToolBar.ToggleButton.End");
+
+				const TSharedRef<SHorizontalBox> RGBA = SNew(SHorizontalBox)
+					+ SHorizontalBox::Slot()
+					.AutoWidth()
+					[
+						GetTextButton("A", ButtonStyleStart, FImageABComparison::EAorB::A)
+					]
+					+ SHorizontalBox::Slot()
+					.AutoWidth()
+					[
+						GetTextButton("B", ButtonStyleEnd, FImageABComparison::EAorB::B)
+					];
+
+				ToolbarBuilder.AddToolBarWidget(RGBA);
+			}
+
+			ToolbarBuilder.AddSeparator();
+		}
+		ToolbarBuilder.EndSection();
+
+		ToolbarBuilder.AddSeparator();
+
+		return ToolbarBuilder.MakeWidget();
+	}
+
+	TSharedRef<SWidget> SImageViewportToolbar::MakeRightToolbar(const TSharedPtr<FExtender>& Extender)
+	{
+		FSlimHorizontalToolBarBuilder ToolbarBuilder = GetToolbarBuilder(Extender);
+
+		ToolbarBuilder.BeginSection("ToolbarRight");
+		{
+			// This is deliberately left empty.
+			// It allows for adding toolbar extensions though.
+		}
+
+		ToolbarBuilder.EndSection();
+
+		return ToolbarBuilder.MakeWidget();
+	}
+
+	FText SImageViewportToolbar::GetZoomMenuLabel() const
+	{
+		const FImageViewportController::FZoomSettings ZoomSettings = ViewportClient->GetZoom();
+
+		const float LayoutAndDPIScaleFactor = GetLayoutAndDPIScaleFactor.Execute();
+
+		const double Zoom = ZoomSettings.Zoom / LayoutAndDPIScaleFactor;
+		FNumberFormattingOptions FormattingOptions;
+		FormattingOptions.SetMaximumFractionalDigits(Zoom < 1.0 ? (Zoom < 0.1 ? 2 : 1) : 0);
+		const FText ZoomPercentage = FText::AsPercent(Zoom, &FormattingOptions);
+
+		if (ZoomSettings.Mode == FImageViewportController::EZoomMode::Custom)
+		{
+			return ZoomPercentage;
+		}
+
+		return FText::Format(
+			LOCTEXT("ZoomFitFill", "{0} {1}"),
+			FText::FromString(ZoomSettings.Mode == FImageViewportController::EZoomMode::Fit ? "Fit" : "Fill"),
+			ZoomPercentage);
+	}
+
+	TSharedRef<SWidget> SImageViewportToolbar::MakeZoomMenu() const
+	{
+		FMenuBuilder MenuBuilder(true, CommandList, nullptr, false, &FAppStyle::Get(), false);
+
+		const auto& Commands = FImageWidgetsCommands::Get();
+
+		MenuBuilder.AddMenuEntry(Commands.Zoom12);
+		MenuBuilder.AddMenuEntry(Commands.Zoom25);
+		MenuBuilder.AddMenuEntry(Commands.Zoom50);
+		MenuBuilder.AddMenuEntry(Commands.Zoom100);
+		MenuBuilder.AddMenuEntry(Commands.Zoom200);
+		MenuBuilder.AddMenuEntry(Commands.Zoom400);
+		MenuBuilder.AddMenuEntry(Commands.Zoom800);
+		MenuBuilder.AddSeparator();
+		MenuBuilder.AddMenuEntry(Commands.ZoomFit);
+		MenuBuilder.AddMenuEntry(Commands.ZoomFill);
+
+		return MenuBuilder.MakeWidget();
+	}
+
+	EVisibility SImageViewportToolbar::GetZoomMenuVisibility() const
+	{
+		return !GetOverlaySettings.Execute().bDisableZoomButton ? EVisibility::Visible : EVisibility::Collapsed;
+	}
+
+	EVisibility SImageViewportToolbar::GetMipMenuVisibility() const
+	{
+		return !GetOverlaySettings.Execute().bDisableMipButton && NumMips.Execute() > 1 ? EVisibility::Visible : EVisibility::Collapsed;
+	}
+
+	EVisibility SImageViewportToolbar::GetABVisibility() const
+	{
+		return !GetOverlaySettings.Execute().bDisableABComparisonButtons ? EVisibility::Visible : EVisibility::Collapsed;
+	}
+
+	FText SImageViewportToolbar::GetMipMenuLabel() const
+	{
+		const int32 Mip = ViewportClient->GetMipLevel();
+		return FText::Format(LOCTEXT("Mip", "Mip {0}"), Mip == INDEX_NONE ? ImageViewportToolbar_Local::Auto : FText::AsNumber(Mip));
+	}
+
+	TSharedRef<SWidget> SImageViewportToolbar::MakeMipMenu() const
+	{
+		FMenuBuilder MenuBuilder(true, nullptr, nullptr, false, &FAppStyle::Get(), false);
+
+		auto AddMenuEntry = [this, &MenuBuilder](const TAttribute<FText>& Label, const TAttribute<FText>& ToolTip, int32 MipLevel)
+		{
+			MenuBuilder.AddMenuEntry(
+				Label,
+				ToolTip,
+				FSlateIcon(),
+				FUIAction(
+					FExecuteAction::CreateSP(ViewportClient.ToSharedRef(), &FImageViewportClient::SetMipLevel, MipLevel),
+					FCanExecuteAction(),
+					FIsActionChecked::CreateLambda([this, MipLevel]() { return ViewportClient->GetMipLevel() == MipLevel; })
+				),
+				NAME_None,
+				EUserInterfaceActionType::RadioButton
+			);
+		};
+
+		AddMenuEntry(ImageViewportToolbar_Local::Auto, LOCTEXT("MipLevelAuto_Tooltip", "Choose Mip Level automatically"), INDEX_NONE);
+
+		if (NumMips.IsBound())
+		{
+			const int32 Num = NumMips.Execute();
+			if (Num > 1)
+			{
+				MenuBuilder.AddSeparator();
+
+				for (int32 Mip = 0; Mip < Num; ++Mip)
+				{
+					const FText MipText = FText::AsNumber(Mip);
+
+					AddMenuEntry(FText::Format(LOCTEXT("MipLevel", "Mip {0}"), MipText),
+					             FText::Format(LOCTEXT("MipLevel_Tooltip", "Display Mip Level {0}"), MipText),
+					             Mip);
+				}
+			}
+		}
+
+		return MenuBuilder.MakeWidget();
+	}
+}
+
+#undef LOCTEXT_NAMESPACE
diff --git a/Engine/Plugins/Editor/ImageWidgets/Source/ImageWidgets/Private/SImageViewportToolbar.h b/Engine/Plugins/Editor/ImageWidgets/Source/ImageWidgets/Private/SImageViewportToolbar.h
new file mode 100644
--- /dev/null
+++ b/Engine/Plugins/Editor/ImageWidgets/Source/ImageWidgets/Private/SImageViewportToolbar.h
@@ -0,0 +1,80 @@
+﻿// Copyright Epic Games, Inc. All Rights Reserved.
+
+#pragma once
+
+#include "CoreMinimal.h"
+#include "ImageViewportClient.h"
+#include "SImageViewport.h"
+#include "SViewportToolBar.h"
+
+class FExtender;
+class FImageViewportCommands;
+class FSlimHorizontalToolBarBuilder;
+class FUICommandList;
+class SViewportToolBarComboMenu;
+
+namespace UE::ImageWidgets
+{
+	class FImageABComparison;
+
+	/**
+	 * Extendable toolbar for the image viewport.
+	 */
+	class SImageViewportToolbar : public SViewportToolBar
+	{
+	public:
+
+		/// Callbacks to image viewport to retrieve necessary image metadata.
+		DECLARE_DELEGATE_RetVal(bool, FHasImage)
+		DECLARE_DELEGATE_RetVal(int32, FNumMips)
+		DECLARE_DELEGATE_RetVal(FGuid, FImageGuid)
+		DECLARE_DELEGATE_RetVal(SImageViewport::FOverlaySettings, FGetOverlaySettings)
+
+		SLATE_BEGIN_ARGS(SImageViewportToolbar)	{}
+		SLATE_END_ARGS()
+
+		/** Parameter for constructing the toolbar. */
+		struct FConstructParameters
+		{
+			FHasImage HasImage;
+			FNumMips NumMips;
+			FImageGuid ImageGuid;
+			FGetLayoutAndDPIScaleFactor GetLayoutAndDPIScaleFactor;
+			FGetOverlaySettings GetOverlaySettings;
+			FImageABComparison* ABComparison = nullptr;
+			TSharedPtr<FExtender> ToolbarExtender;
+			const ISlateStyle* ToolbarExtenderStyle = nullptr;
+			FName ToolbarExtenderStyleName;
+		};
+
+		void Construct(const FArguments& InArgs, const TSharedPtr<FImageViewportClient>& InViewportClient, const TSharedPtr<FUICommandList>& InCommandList,
+		               FConstructParameters Parameters);
+
+	private:
+		FSlimHorizontalToolBarBuilder GetToolbarBuilder(const TSharedPtr<FExtender>& Extender) const;
+		TSharedRef<SWidget> MakeLeftToolbar(const TSharedPtr<FExtender>& Extender);
+		TSharedRef<SWidget> MakeCenterToolbar(const TSharedPtr<FExtender>& Extender, bool bEnableABComparison);
+		TSharedRef<SWidget> MakeRightToolbar(const TSharedPtr<FExtender>& Extender);
+
+		EVisibility GetZoomMenuVisibility() const;
+		EVisibility GetMipMenuVisibility() const;
+		EVisibility GetABVisibility() const;
+		FText GetMipMenuLabel() const;
+		TSharedRef<SWidget> MakeMipMenu() const;
+
+		FText GetZoomMenuLabel() const;
+		TSharedRef<SWidget> MakeZoomMenu() const;
+
+		TSharedPtr<FImageViewportClient> ViewportClient;
+		TSharedPtr<FUICommandList> CommandList;
+
+		FHasImage HasImage;
+		FNumMips NumMips;
+		FImageGuid ImageGuid;
+		FGetLayoutAndDPIScaleFactor GetLayoutAndDPIScaleFactor;
+		FGetOverlaySettings GetOverlaySettings;
+		FImageABComparison* ABComparison = nullptr;
+		const ISlateStyle* ToolbarExtenderStyle = nullptr;
+		FName ToolbarExtenderStyleName;
+	};
+}
diff --git a/Engine/Plugins/Editor/ImageWidgets/Source/ImageWidgets/Public/IImageViewer.h b/Engine/Plugins/Editor/ImageWidgets/Source/ImageWidgets/Public/IImageViewer.h
new file mode 100644
--- /dev/null
+++ b/Engine/Plugins/Editor/ImageWidgets/Source/ImageWidgets/Public/IImageViewer.h
@@ -0,0 +1,147 @@
+// Copyright Epic Games, Inc. All Rights Reserved.
+
+#pragma once
+
+#include "CoreMinimal.h"
+#include "Misc/TVariant.h"
+
+class FCanvas;
+class FViewport;
+
+namespace UE::ImageWidgets
+{
+	/**
+	 * Interface for a component that wants to show image related content with the Slate widgets in the ImageWidget module.
+	 * In this context, an "image" is considered to be any 2D content that is contained within an axis-aligned rectangle.
+	 */
+	class IImageViewer
+	{
+	public:
+		/**
+		 * Information about an image to be displayed.
+		 */
+		struct FImageInfo
+		{
+			/** Unique image identifier. It can encode any helpful metadata as long as no two provided images have the same GUID. */
+			FGuid Guid;
+
+			/** XY size of the image in pixels. */
+			FIntPoint Size;
+
+			/** Number of available MIPs. This should be set to zero if the image type does not support mips. */
+			int32 NumMips;
+
+			/** Indicates that this image is valid for display. */
+			bool bIsValid;
+		};
+
+		/**
+		 * Information necessary for correctly drawing an image.
+		 */
+		struct FDrawProperties
+		{
+			/**
+			 * Where in the 2D plane the image rectangle is supposed to be drawn.
+			 */
+			struct FPlacement
+			{
+				/** Offset from the origin, i.e. (0, 0). */
+				FVector2d Offset;
+
+				/** XY size of the axis aligned rectangle containing the image. */
+				FVector2d Size;
+
+				/**
+				 * The zoom factor used for the image.
+				 * While this might not be necessary for drawing the image, it can be helpful in certain use cases. For example, interpolation could explicitly
+				 * be turned off when zooming into a texture to show the discrete pixel contents of the texture instead of the interpolated result.
+				 */
+				double ZoomFactor;
+			};
+
+			/**
+			 * Information about MIP levels. This can be ignored if the image type does not support MIPs. 
+			 */
+			struct FMip
+			{
+				/** The selected MIP level. */
+				float MipLevel;
+			};
+
+			/**
+			 * Information necessary for rendering AB comparisons.
+			 */
+			struct FABComparison
+			{
+				/** @return true if an AB comparison should be drawn instead of a single image. */
+				bool IsActive() const { return GuidA.IsValid() && GuidB.IsValid(); };
+
+				/** Unique identifier for image A. */
+				FGuid GuidA;
+				
+				/** Unique identifier for image B. */
+				FGuid GuidB;
+
+				/**
+				 * Value between 0..1 to indicate where the threshold between images A and B is.
+				 * A value of 0 means that only B should be drawn.
+				 * A value of 0.5 means that the left half of A and the right half of B should be drawn.
+				 * A value of 1 means that only A should be drawn. 
+				 */
+				double Threshold;
+			};
+
+			FPlacement Placement;
+			FMip Mip;
+			FABComparison ABComparison;
+		};
+
+		/**
+		 * Provides any necessary metadata for the image widgets about the image that is currently supposed to be displayed. This data is generic in the sense
+		 * that the image widgets don't need to know any of the image structure, its content or how to draw it. Instead, the image drawing is done directly
+		 * via @see DrawCurrentImage.
+		 * @return information about the image to be displayed
+		 */
+		virtual FImageInfo GetCurrentImageInfo() const = 0;
+
+		/**
+		 * Draws the image that is currently supposed to be displayed within the 2D viewport.
+		 * @param Viewport The viewport the image is drawn into
+		 * @param Canvas The canvas used for drawing the image
+		 * @param Properties Information for drawing the image based on the current image viewport state 
+		 */
+		virtual void DrawCurrentImage(FViewport* Viewport, FCanvas* Canvas, const FDrawProperties& Properties) = 0;
+		
+		/**
+		 * Provides information about a given pixel.
+		 * @param PixelCoords XY coordinates for the pixel
+		 * @param MipLevel MIP level that is currently displayed; this can be ignored for images not supporting MIPs 
+		 * @return Either a color value in byte or float format, i.e. FColor or FLinearColor, or no value if there is no valid pixel at the provided coordinates 
+		 */
+		virtual TOptional<TVariant<FColor, FLinearColor>> GetCurrentImagePixelColor(FIntPoint PixelCoords, int32 MipLevel) const = 0;
+
+		/**
+		 * Notifies about the image with the given GUID being selected.
+		 * This can be implemented as an empty function if the image viewer implementation does not support switching between different images.
+		 * @param Guid Unique identifier of the selected image
+		 */
+		virtual void OnImageSelected(const FGuid& Guid) = 0;
+
+		/**
+		 * Returns if a given GUID represents a currently available image.
+		 * @param Guid Unique identifier of the potentially available image
+		 * @return true if the GUID represents a currently available image
+		 */
+		virtual bool IsValidImage(const FGuid& Guid) const = 0;
+
+		/**
+		 * Returns the name of a currently available image. This is potentially used in the UI to refer to an image.
+		 * @param Guid Unique identifier of the potentially available image
+		 * @return Name of the image if the GUID represents a currently available image
+		 */
+		virtual FText GetImageName(const FGuid& Guid) const = 0;
+
+	protected:
+		~IImageViewer() = default;
+	};
+}
diff --git a/Engine/Plugins/Editor/ImageWidgets/Source/ImageWidgets/Public/SImageCatalog.h b/Engine/Plugins/Editor/ImageWidgets/Source/ImageWidgets/Public/SImageCatalog.h
new file mode 100644
--- /dev/null
+++ b/Engine/Plugins/Editor/ImageWidgets/Source/ImageWidgets/Public/SImageCatalog.h
@@ -0,0 +1,374 @@
+// Copyright Epic Games, Inc. All Rights Reserved.
+
+#pragma once
+
+#include "CoreMinimal.h"
+#include "Framework/Views/ITypedTableView.h"
+#include "Widgets/SCompoundWidget.h"
+
+namespace UE::ImageWidgets
+{
+	/**
+	 * Contains all data for a catalog item.
+	 */
+	struct FImageCatalogItemData
+	{
+		IMAGEWIDGETS_API FImageCatalogItemData(FGuid Guid, const FSlateBrush& Brush, const FText& Name, const FText& Info, const FText& ToolTip);
+
+		/** Unique identifier for the catalog item */
+		FGuid Guid;
+
+		/** Brush used for displaying the item's thumbnail */
+		FSlateBrush Thumbnail;
+
+		/** Name of the item */
+		FText Name;
+
+		/** Auxiliary information for the item */
+		FText Info;
+
+		/** Tooltip that is shown when hovering over any part of the item's widget in the catalog */
+		FText ToolTip;
+	};
+
+	/**
+	 * Generic catalog widget for listing and interacting with 2D image-like content.
+	 * Each catalog item is represented by its own widget based on its @see FImageCatalogItemData.
+	 * Entries in the catalog can be assigned to customizable groups.
+	 */
+	class SImageCatalog : public SCompoundWidget
+	{
+	public:
+		/**
+		 * Delegate that gets called when an item is selected in the catalog.
+		 * The given @see FGuid identifies the item that was selected.
+		 */
+		DECLARE_DELEGATE_OneParam(FOnItemSelected, const FGuid&)
+
+		/**
+		 * Delegate that gets called for creating a context menu for a group.
+		 * Return @see SNullWidget::NullWidget to not show a context menu.
+		 */
+		DECLARE_DELEGATE_RetVal_OneParam(TSharedPtr<SWidget>, FOnGetGroupContextMenu, FName)
+
+		/**
+		 * Delegate that gets called for creating a context menu for a set of selected items.
+		 * Return @see SNullWidget::NullWidget to not show a context menu.
+		 */
+		DECLARE_DELEGATE_RetVal_OneParam(TSharedPtr<SWidget>, FOnGetItemsContextMenu, const TArray<FGuid>&)
+
+		SLATE_BEGIN_ARGS(SImageCatalog)
+				: _DefaultGroupName(NAME_None)
+				, _DefaultGroupHeading(FText())
+				, _SelectionMode(ESelectionMode::Multi)
+				, _bAllowSelectionAcrossGroups(false)
+				, _bShowEmptyGroups(false)
+			{
+			}
+
+			/** Name of the default group, i.e. the group that gets used when no group is explicitly specified. */
+			SLATE_ARGUMENT(FName, DefaultGroupName)
+
+			/** Header text for the default group. */
+			SLATE_ARGUMENT(FText, DefaultGroupHeading)
+
+			/** Delegate that gets called when an item is selected in the catalog. */
+			SLATE_EVENT(FOnItemSelected, OnItemSelected)
+
+			/** Delegate that gets called for creating a context menu for a group. */
+			SLATE_EVENT(FOnGetGroupContextMenu, OnGetGroupContextMenu)
+
+			/** Delegate that gets called for creating a context menu for a set of selected items. */
+			SLATE_EVENT(FOnGetItemsContextMenu, OnGetItemsContextMenu)
+
+			/** Defines the selection behavior within an item list, e.g. only allow single item selection or do not allow any selection. */
+			SLATE_ARGUMENT(ESelectionMode::Type, SelectionMode)
+
+			/** When an item is selected, the selection in other groups will be cleared unless this flag is set to true. */
+			SLATE_ARGUMENT(bool, bAllowSelectionAcrossGroups)
+
+			/** Empty groups will be hidden unless this flag is set to true. */
+			SLATE_ARGUMENT(bool, bShowEmptyGroups)
+		SLATE_END_ARGS()
+
+		/**
+		 * Function used by Slate to construct the image catalog widget with the given arguments.
+		 * @param Args Slate arguments defined above
+		 */
+		IMAGEWIDGETS_API void Construct(const FArguments& Args);
+
+		/**
+		 * @returns The name of the default group.
+		 */
+		IMAGEWIDGETS_API FName GetDefaultGroupName() const;
+
+		/**
+		 * Add a custom group to the catalog.
+		 * @param Name Unique identifier for the group.
+		 * @param Heading Header text to be used for the group; no header will be shown if this is empty.
+		 * @returns False if a group with the given name already exits, otherwise true.
+		 */
+		IMAGEWIDGETS_API bool AddGroup(FName Name, const FText& Heading);
+
+		/**
+		 * Add a custom group to the catalog before an already existing group.
+		 * @param Name Unique identifier for the group.
+		 * @param Heading Header text to be used for the group; no header will be shown if this is empty.
+		 * @param BeforeGroupWithThisName Name of an already existing group before which the new custom group is going to be added.
+		 * @returns False if a group with the given name already exits or if the group before which the new group is supposed to be added does not exist,
+		 * otherwise true.
+		 */
+		IMAGEWIDGETS_API bool AddGroup(FName Name, const FText& Heading, FName BeforeGroupWithThisName);
+
+		/**
+		 * Set the header text for an existing group.
+		 * @param Name Unique identifier for the group.
+		 * @param Heading Header text to be used for the group; no header will be shown if this is empty.
+		 * @returns False if a group with the given name does not exist, otherwise true.
+		 */
+		IMAGEWIDGETS_API bool SetGroupHeading(FName Name, const FText& Heading);
+
+		/**
+		* Remove an existing group and any items in the group.
+		* @param Name Unique identifier of the group; must not be the default group name.
+		* @returns An empty value if a group with the given name does not exist, otherwise a list of unique identifiers for the items that were in the group
+		* and got removed from the catalog.
+		*/
+		IMAGEWIDGETS_API TOptional<TArray<FGuid>> RemoveGroup(FName Name);
+
+		/**
+		* Remove an existing group, and moves any items in the group into another group.
+		* Note that if the other group does not exist, the items get moved into the default group instead.
+		* @param Name Unique identifier of the group; must not be the default group name.
+		* @param GroupToMoveItemsInto Name of an already existing group in which the items from the removed group get moved into.
+		* @returns An empty value if a group with the given name does not exist, otherwise a list of unique identifiers for the items that were in the group
+		* and got moved into the other group.
+		*/
+		IMAGEWIDGETS_API TOptional<TArray<FGuid>> RemoveGroup(FName Name, FName GroupToMoveItemsInto);
+
+		/**
+		 * Return the number of groups in the catalog including the default group.
+		 * @returns The number of groups in the catalog.
+		 */
+		IMAGEWIDGETS_API int32 NumGroups() const;
+
+		/**
+		 * Return the name of the group at a given position.
+		 * @param Index Position of the group, with 0 being the first, and @see NumGroups() - 1 being the last.
+		 * @returns An empty value if the given position is invalid, otherwise the name of the group at that position.
+		 */
+		IMAGEWIDGETS_API TOptional<FName> GetGroupNameAt(int32 Index) const;
+
+		/**
+		 * Add an item to the default group.
+		 * @param Item Data for the item that is being added.
+		 * @returns True if the item was added successfully.
+		 */
+		IMAGEWIDGETS_API bool AddItem(const TSharedPtr<FImageCatalogItemData>& Item);
+
+		/**
+		 * Add an item to the default group right before an existing item.
+		 * If the unique identifier for the other item is invalid or not in the default group, the item will be added at the end of the default group instead.
+		 * @param Item Data for the item that is being added.
+		 * @param BeforeItemWithThisGuid Unique identifier of the other item before which the new item should be added.
+		 * @returns True if the item was added successfully.
+		 */
+		IMAGEWIDGETS_API bool AddItem(const TSharedPtr<FImageCatalogItemData>& Item, const FGuid& BeforeItemWithThisGuid);
+
+		/**
+		 * Add an item to an existing group.
+		 * @param Item Data for the item that should be added.
+		 * @param Group Name of the existing group to which the item should be added. 
+		 * @returns False if a group with the given name does not exist, otherwise true.
+		 */
+		IMAGEWIDGETS_API bool AddItem(const TSharedPtr<FImageCatalogItemData>& Item, FName Group);
+
+		/**
+		 * Add an item to an existing group right before an existing item.
+		 * If the unique identifier for the other item is invalid or not in the given group, the item will be added at the end of the group instead.
+		 * @param Item Data for the item that should be added.
+		 * @param Group Name of the existing group to which the item should be added.
+		 * @param BeforeItemWithThisGuid Unique identifier of the other item before which the new item should be added.
+		 * @returns False if a group with the given name does not exist, otherwise true.
+		 */
+		IMAGEWIDGETS_API bool AddItem(const TSharedPtr<FImageCatalogItemData>& Item, FName Group, const FGuid& BeforeItemWithThisGuid);
+
+		/**
+		 * Move an already existing item before another item within the same group.
+		 * @param Guid The unique identifier of the existing item that should be moved.
+		 * @param BeforeItemWithThisGuid Unique identifier of the other item before which the item should be moved.
+		 * @returns False if the item to be moved or the item before which it should be moved does not exist or if both are not in the same group,
+		 * otherwise true.
+		 */
+		IMAGEWIDGETS_API bool MoveItem(const FGuid& Guid, const FGuid& BeforeItemWithThisGuid);
+
+		/**
+		 * Move an already existing item to another group.
+		 * @param Guid The unique identifier of the existing item that should be moved.
+		 * @param Group Name of the existing group to which the item should be moved.
+		 * @returns False if an item with the given unique identifier does not exist, a group with the given name does not exist, or the item is already in the
+		 * given group, otherwise true.
+		 */
+		IMAGEWIDGETS_API bool MoveItem(const FGuid& Guid, FName Group);
+
+		/**
+		 * Move an already existing item to another group, before another item within that group.
+		 * @param Guid The unique identifier of the existing item that should be moved.
+		 * @param Group Name of the existing group to which the item should be moved.
+		 * @param BeforeItemWithThisGuid Unique identifier of the other item before which the item should be moved.
+		 * @returns False if an item with the given unique identifier does not exist, a group with the given name does not exist, or the item is already in the
+		 * given group, otherwise true.
+		 */
+		IMAGEWIDGETS_API bool MoveItem(const FGuid& Guid, FName Group, const FGuid& BeforeItemWithThisGuid);
+
+		/**
+		 * Remove an existing item.
+		 * @param Guid The unique identifier of the existing item.
+		 * @returns False if an item with the given unique identifier does not exist, otherwise true.
+		 */
+		IMAGEWIDGETS_API bool RemoveItem(const FGuid& Guid);
+
+		/**
+		 * Return the total number of items in the catalog.
+		 * @returns The number of items in the catalog across all groups.
+		 */
+		IMAGEWIDGETS_API int32 NumItems() const;
+
+		/**
+		 * Return the number of items in a group.
+		 * @param Group Name of an existing group.
+		 * @returns The number of items in the given group.
+		 */
+		IMAGEWIDGETS_API int32 NumItems(FName Group) const;
+
+		/**
+		 * Retrieve the existing item for a given unique identifier.
+		 * @param Guid The unique identifier of the item.
+		 * @return The pointer to the item or an invalid pointer if no item with the given unique identifier exists.
+		 */
+		IMAGEWIDGETS_API TSharedPtr<const FImageCatalogItemData> GetItem(const FGuid& Guid) const;
+
+		/**
+		 * Return the group an existing item belongs to.
+		 * @param Guid The unique identifier of the existing item.
+		 * @returns An empty value if an item with the given unique identifier does not exist, otherwise the name of the group the item belongs to.
+		 */
+		IMAGEWIDGETS_API TOptional<FName> GetItemGroupName(const FGuid& Guid) const;
+
+		/**
+		 * Return the index of an existing item within the group it belongs to.
+		 * @param Guid The unique identifier of the existing item.
+		 * @returns An empty value if no item with the given unique identifier exists, otherwise index of the item in the group.
+		 */
+		IMAGEWIDGETS_API TOptional<int32> GetItemIndex(const FGuid& Guid) const;
+
+		/**
+		 * Return an existing item's group and the index within that group.
+		 * @param Guid The unique identifier of the existing item.
+		 * @returns An empty value if an item with the given unique identifier does not exist, otherwise the name of the group the item belongs to and
+		 * the index of the item in the group.
+		 */
+		IMAGEWIDGETS_API TOptional<TTuple<FName, int32>> GetItemGroupNameAndIndex(const FGuid& Guid) const;
+
+		/**
+		 * Retrieve the existing item for a given index within the default group.
+		 * @param Index Index of the item in the default group; should be at least 0 and less than @see NumItems(@see GetDefaultGroup()).
+		 * @returns The pointer to the item or an invalid pointer if the index is invalid.
+		 */
+		IMAGEWIDGETS_API TSharedPtr<const FImageCatalogItemData> GetItemAt(int32 Index) const;
+
+		/**
+		* Retrieve the existing item within a given group for a given index.
+		 * @param Index Index of the item in the given group; should be at least 0 and less than @see NumItems(Group).
+		 * @param Group Name of an existing group.
+		 * @returns The pointer to the item or an invalid pointer if the group does not exist or the index is invalid.
+		 */
+		IMAGEWIDGETS_API TSharedPtr<const FImageCatalogItemData> GetItemAt(int32 Index, FName Group) const;
+
+		/**
+		 * Retrieve the unique identifier of an existing item at the given index within the default group.
+		 * @param Index Index of the item in the default group; should be at least 0 and less than @see NumItems(@see GetDefaultGroup()).
+		 * @return Unique identifier of the item or no value if the index is invalid.
+		 */
+		IMAGEWIDGETS_API TOptional<FGuid> GetItemGuidAt(int32 Index) const;
+
+		/**
+		 * Retrieve the unique identifier of an existing item at the given index within the given group.
+		 * @param Index Index of the item in the given group; should be at least 0 and less than @see NumItems(Group).
+		 * @param Group Name of an existing group.
+		 * @return Unique identifier of the item or no value if the group does not exist or the index is invalid.
+		 */
+		IMAGEWIDGETS_API TOptional<FGuid> GetItemGuidAt(int32 Index, FName Group) const;
+
+		/**
+		 * Select an existing item.
+		 * @param Guid The unique identifier of the existing item.
+		 * @returns False if an item with the given unique identifier does not exist, otherwise true.
+		 */
+		IMAGEWIDGETS_API bool SelectItem(const FGuid& Guid);
+
+		/**
+		 * Deselect an existing item.
+		 * @param Guid The unique identifier of the existing item.
+		 * @returns False if an item with the given unique identifier does not exist, otherwise true.
+		 */
+		IMAGEWIDGETS_API bool DeselectItem(const FGuid& Guid);
+
+		/**
+		 * Clear any selection in the catalog, i.e. across all groups.
+		 */
+		IMAGEWIDGETS_API void ClearSelection();
+
+		/**
+		 * Clear the selection for a given group; selections in other groups will be unchanged.
+		 * @param Group Name of an existing group.
+		 * @returns False if a group with the given name does not exist
+		 */
+		IMAGEWIDGETS_API bool ClearSelection(FName Group);
+
+		/**
+		 * Update an existing item's data. The item data needs to contain the item's unique identifier.
+		 * @param Item Data for the exiting item that should be updated, including its unique identifier.
+		 * @returns False if an item with the given unique identifier does not exist, otherwise true.
+		 */
+		IMAGEWIDGETS_API bool UpdateItem(const FImageCatalogItemData& Item);
+
+		/**
+		 * Update the info text of an existing item.
+		 * @param Guid The unique identifier of the existing item.
+		 * @param Info Text for the info label in the item's widget.
+		 * @returns False if an item with the given unique identifier does not exist, otherwise true.
+		 */
+		IMAGEWIDGETS_API bool UpdateItemInfo(const FGuid& Guid, const FText& Info);
+
+		/**
+		 * Update the name of an existing item.
+		 * Nothing happens if no item with the given unique identifier exists in the catalog.
+		 * @param Guid The unique identifier of the existing item.
+		 * @param Name Text for the name label in the item's widget.
+		 * @returns False if an item with the given unique identifier does not exist, otherwise true.
+		 */
+		IMAGEWIDGETS_API bool UpdateItemName(const FGuid& Guid, const FText& Name);
+
+		/**
+		 * Update the thumbnail of an existing item.
+		 * Nothing happens if no item with the given unique identifier exists in the catalog.
+		 * @param Guid The unique identifier of the existing item.
+		 * @param Thumbnail Brush used for the thumbnail in the item's widget.
+		 * @returns False if an item with the given unique identifier does not exist, otherwise true.
+		 */
+		IMAGEWIDGETS_API bool UpdateItemThumbnail(const FGuid& Guid, const FSlateBrush& Thumbnail);
+
+		/**
+		 * Update the tooltip text of an existing item.
+		 * Nothing happens if no item with the given unique identifier exists in the catalog.
+		 * @param Guid The unique identifier of the existing item.
+		 * @param ToolTip Text for the tooltip label in the item's widget.
+		 * @returns False if an item with the given unique identifier does not exist, otherwise true.
+		 */
+		IMAGEWIDGETS_API bool UpdateItemToolTip(const FGuid& Guid, const FText& ToolTip);
+
+	private:
+		TPimplPtr<class FImpl> Impl;
+	};
+}
diff --git a/Engine/Plugins/Editor/ImageWidgets/Source/ImageWidgets/Public/SImageViewport.h b/Engine/Plugins/Editor/ImageWidgets/Source/ImageWidgets/Public/SImageViewport.h
new file mode 100644
--- /dev/null
+++ b/Engine/Plugins/Editor/ImageWidgets/Source/ImageWidgets/Public/SImageViewport.h
@@ -0,0 +1,311 @@
+﻿// Copyright Epic Games, Inc. All Rights Reserved.
+
+#pragma once
+
+#include "SEditorViewport.h"
+#include "Framework/MultiBox/MultiBoxExtender.h"
+
+class SViewportToolBar;
+
+namespace UE::ImageWidgets
+{
+	class FImageABComparison;
+	class FImageViewportClient;
+	class FStatusBarExtension;
+	class IImageViewer;
+	class SImageViewportToolbar;
+
+	/**
+	 * Generic viewport for displaying and interacting with 2D image-like content.
+	 * The drawing of the images is deferred to an @see IImageViewer implementation that needs to be provided upon construction. This viewport only uses the
+	 * metadata provided by the image viewer to have sufficient information about the image without being aware of its actual format or contents.
+	 */
+	class SImageViewport : public SEditorViewport
+	{
+	public:
+		/**
+		 * Contains any extensions widgets for the viewport toolbar.
+		 * This implementation is similar to the existing toolbar extender, i.e. FExtender, which is designed to only work with dedicated toolbar builders and
+		 * can currently not be used with status bars.
+		 */
+		class FStatusBarExtender
+		{
+		public:
+			/**
+			 * Delegate that gets called for each extension when the status bar is constructed.
+			 * The given @see SHorizontalBox is the layout container that the widget for the extension needs to be added to by calling
+			 * @see SHorizontalBox::AddSlot().
+			 */
+			DECLARE_DELEGATE_OneParam(FDelegate, SHorizontalBox&);
+
+			/**
+			 * Adds an extension to the status bar.
+			 * @param ExtensionHook The section the extension is applied to
+			 * @param HookPosition Where in relation to the section the extension is applied to
+			 * @param Commands The UI command list responsible for handling actions used in the widgets added in the extension
+			 * @param Delegate Delegate called for adding the extension widgets when the status bar is constructed
+			 */
+			IMAGEWIDGETS_API void AddExtension(FName ExtensionHook, EExtensionHook::Position HookPosition, const TSharedPtr<FUICommandList>& Commands,
+			                                   const FDelegate& Delegate);
+
+		private:
+			friend SImageViewport;
+
+			/** Used by the viewport to add extensions to the status bar. */
+			void Apply(FName ExtensionHook, EExtensionHook::Position HookPosition, SHorizontalBox& HorizontalBox) const;
+
+			/** List of extensions that get applied to the viewport status bar. */
+			TArray<TPimplPtr<FStatusBarExtension>> Extensions;
+		};
+
+		/**
+		 * Settings related to drawing viewport contents other than the image itself.
+		 */
+		struct FDrawSettings
+		{
+			/** Clear color for the viewport. */
+			FLinearColor ClearColor = FLinearColor::Black;
+
+			/** Flag that enables drawing a border around the image. If enabled, the border is drawn underneath the image.
+			 *  The center of the line drawn for the border is at the exact edge of the image, i.e. when the image is drawn on top of it, the outer half of the
+			 *  border is visible, and the inner half is occluded by the image. */
+			bool bBorderEnabled = false;
+
+			/** Thickness of the border. */
+			float BorderThickness = 1.0f;
+
+			/** Color of the border */
+			FLinearColor BorderColor = FVector3f(0.2f);
+
+			/** Flag that enables the drawing of the background color within the image rectangle. */
+			bool bBackgroundColorEnabled = false;
+
+			/** Color for the background within the image rectangle. If this is different to the clear color, it shows where the image is even if nothing is drawn. */
+			FLinearColor BackgroundColor = FLinearColor::Black;
+
+			/** Flag that enabled the drawing of a background checker texture. If this is enabled, the texture is drawn on top of the background color. */
+			bool bBackgroundCheckerEnabled = false;
+
+			/** First color used in the background checker texture. */
+			FLinearColor BackgroundCheckerColor1 = FLinearColor::White;
+
+			/** Second color used in the background checker texture. */
+			FLinearColor BackgroundCheckerColor2 = FVector3f(0.8f);
+
+			/** Size of a single square within the background checker texture in pixels. */
+			uint32 BackgroundCheckerSize = 8;
+		};
+
+		/**
+		 * Settings related to viewport controls.
+		 */
+		struct FControllerSettings
+		{
+			enum class EDefaultZoomMode
+			{
+				Fit, // Make the image fit within the viewport, but do not make it larger than the original size.
+				Fill // Make the image fit within the viewport, and if it is smaller than the viewport, zoom in to fill the viewport.
+			};
+
+			/**
+			 * Delegate for custom input key event handling in the viewport client.
+			 * The delegate is called after the image viewports input key events were handled, e.g. zoom or view reset via F key, and before a call
+			 * to @see FEditorViewportClient::InputKey().
+			 *
+			 * @param	EventArgs - The Input event args.
+			 * @return	True to consume the key event, false to pass it on.
+			 */
+			DECLARE_DELEGATE_RetVal_OneParam(bool, FOnInputKey, const FInputKeyEventArgs& EventArgs);
+
+			/** Zoom mode that gets set on viewport construction and whenever the viewport controller is reset via @see ResetController(). */
+			EDefaultZoomMode DefaultZoomMode = EDefaultZoomMode::Fit;
+
+			/** Delegate for custom input key handling. */
+			FOnInputKey OnInputKey;
+
+			/** Adjusts zoom to be proportional to viewport resizes, e.g. doubling the viewport size will double the zoom. */
+			bool bZoomOnResize = false;
+		};
+
+		/**
+		 * Settings related to the viewport overlay.
+		 */
+		struct FOverlaySettings
+		{
+			/** Do not show the zoom button in the left toolbar. */
+			bool bDisableZoomButton = false;
+
+			/** Do not show the MIP button in the left toolbar. */
+			bool bDisableMipButton = false;
+
+			/** Do not show the AB comparison buttons in the center toolbar. */
+			bool bDisableABComparisonButtons = false;
+
+			/** Do not show the bottom left status bar. */
+			bool bDisableStatusBarLeft = false;
+
+			/** Do not show the bottom right status bar. */
+			bool bDisableStatusBarRight = false;
+		};
+
+		SLATE_BEGIN_ARGS(SImageViewport)
+			: _bABComparisonEnabled(false)
+			{
+			}
+
+			/** Extensions for the viewport toolbar; valid hooks are "ToolbarLeft", "ToolbarCenter", "ToolbarRight" */
+			SLATE_ARGUMENT(TSharedPtr<FExtender>, ToolbarExtender)
+
+			/** Style override for viewport toolbar extensions. */
+			SLATE_ARGUMENT(const ISlateStyle*, ToolbarExtenderStyle)
+
+			/**
+			 * Name used for the viewport toolbar extensions style override.
+			 * Only used if ToolbarExtenderStyle is set. If this is not set, "EditorViewportToolBar" will be used.
+			 */
+			SLATE_ARGUMENT(FName, ToolbarExtenderStyleName)
+
+			/** Extensions for the viewport status bar; valid hooks are "StatusBarLeft", "StatusBarCenter", "StatusBarRight" */
+			SLATE_ARGUMENT(TSharedPtr<FStatusBarExtender>, StatusBarExtender)
+
+			/** Settings for drawing viewport contents other than the actual image */
+			SLATE_ATTRIBUTE(FDrawSettings, DrawSettings)
+
+			/** Settings for customizing the viewport overlay */
+			SLATE_ATTRIBUTE(FOverlaySettings, OverlaySettings)
+
+			/** Enables AB comparison controls in the toolbar */
+			SLATE_ARGUMENT(bool, bABComparisonEnabled)
+
+			/** Settings for controlling the viewport */
+			SLATE_ARGUMENT(FControllerSettings, ControllerSettings)
+
+		SLATE_END_ARGS()
+
+		IMAGEWIDGETS_API SImageViewport();
+		IMAGEWIDGETS_API virtual ~SImageViewport() override;
+
+		/**
+		 * Function used by Slate to construct the image viewport widget with the given arguments.
+		 * @param InArgs Slate arguments defined above
+		 * @param InImageViewer @see IImageViewer implementation that holds and draws the actual image contents.
+		 */
+		IMAGEWIDGETS_API void Construct(const FArguments& InArgs, const TSharedRef<IImageViewer>& InImageViewer);
+
+		/**
+		 * Provides access to the viewport toolbar, which is needed to dynamically generate certain toolbar widgets, e.g. an @see SEditorViewportToolbarMenu,
+		 * as part of a toolbar extension. 
+		 * @return Pointer to the viewport toolbar 
+		 */
+		IMAGEWIDGETS_API TSharedPtr<SViewportToolBar> GetParentToolbar() const;
+
+		/**
+		 * Result for calls to @see GetPixelCoordinatesUnderCursor.
+		 */
+		struct FPixelCoordinatesUnderCursorResult
+		{
+			/** Indicates that the cursor position is currently valid. This is set to false, for example, when the cursor is outside the widget. */
+			bool bIsValid = false;
+
+			/** Pixel coordinates under the cursor relative to the image rectangle size and placement. Note that the coordinates might be outside the image
+			 * rectangle, i.e. values might be negative or larger than the image size. */
+			FVector2d Coordinates;
+		};
+
+		/**
+		 * Provides the pixel coordinates under the cursor.
+		 * @return Pixel coordinates result under the cursor; see @see FPixelCoordinatesUnderCursorResult for more details
+		 */
+		IMAGEWIDGETS_API FPixelCoordinatesUnderCursorResult GetPixelCoordinatesUnderCursor() const;
+
+		/**
+		 * Request to redraw the viewport contents as soon as possible.
+		 */
+		IMAGEWIDGETS_API void RequestRedraw() const;
+
+		/**
+		 * Resets the camera controller to default values.
+		 * @param ImageSize Size of the currently displayed image in pixels
+		 */
+		IMAGEWIDGETS_API void ResetController(FIntPoint ImageSize) const;
+
+		/**
+		 * Resets the MIP level to the default.
+		 */
+		IMAGEWIDGETS_API void ResetMip() const;
+
+		/**
+		 * Resets the zoom to default.
+		 * @param ImageSize Size of the currently displayed image in pixels
+		 */
+		IMAGEWIDGETS_API void ResetZoom(FIntPoint ImageSize) const;
+
+		// SEditorViewport overrides - begin
+	public:
+		virtual bool IsVisible() const override { return true; } // Overriding this method avoids the 0.25 seconds update delay in SEditorViewport.
+	protected:
+		virtual void BindCommands() override;
+		virtual TSharedRef<FEditorViewportClient> MakeEditorViewportClient() override;
+		virtual void PopulateViewportOverlays(TSharedRef<SOverlay> Overlay) override;
+		// SEditorViewport overrides - end
+
+		// SCompoundWidget overrides - begin
+		virtual void OnArrangeChildren(const FGeometry& AllottedGeometry, FArrangedChildren& ArrangedChildren) const override;
+		// SCompoundWidget overrides - end
+
+	private:
+		TSharedRef<SWidget> MakeViewportToolbarOverlay();
+
+		/** Create the status bar widgets. */
+		TSharedRef<SWidget> MakeStatusBar(TSharedRef<SOverlay> Overlay);
+
+		/** Provides the window DPI scale factor. */
+		float GetDPIScaleFactor() const;
+
+		/** Provides the external layout factor (ratio of accumulated scale to window DPI). */
+		float GetLayoutScaleFactor() const;
+
+		/** Provides the layout and DPI scale factor, i.e. the full accumulated geometry scale. */
+		float GetLayoutAndDPIScaleFactor() const;
+
+		/** Provides the text for the color picker display in the status bar. */
+		FText GetPickerLabel() const;
+
+		/** Provides the text for the resolution display in the status bar. */
+		FText GetResolutionLabel() const;
+
+		/** Makes the draw settings available either as fixed values or via a callback to the outside of the viewport. */
+		TAttribute<FDrawSettings> DrawSettings;
+
+		/** Makes the overlay settings available either as fixed values or via a callback to the outside of the viewport. */
+		TAttribute<FOverlaySettings> OverlaySettings;
+
+		/** Flag that determines if AB comparison widgets are enabled or not. The value does not change after the call to @see Construct(). */
+		bool bABComparisonEnabled = false;
+
+		/** Data and logic related to AB comparisons. This is effectively unused when @see bABComparisonEnabled is set to false. */
+		TPimplPtr<FImageABComparison> ABComparison;
+
+		/** The image viewer that holds and draws the actual images. */
+		TSharedPtr<IImageViewer> ImageViewer;
+
+		/** The viewport client that takes care of camera controls and displaying the viewport contents. */
+		TSharedPtr<FImageViewportClient> ImageViewportClient;
+
+		/** The toolbar that controls some of the behavior of the viewport and optionally also the image viewer via toolbar extensions. */
+		TSharedPtr<SImageViewportToolbar> ImageViewportToolbar;
+
+		/** Toolbar extensions provided by the call to @see Construct(). This pointer is reset after the extensions were applied during construction. */
+		TSharedPtr<FExtender> ToolbarExtender;
+
+		/** Style override for toolbar extender. Only used during construction, then reset. */
+		const ISlateStyle* ToolbarExtenderStyle = nullptr;
+		FName ToolbarExtenderStyleName;
+
+		/** Status bar extensions provided by the call to @see Construct(). This pointer is reset after the extensions were applied during construction. */
+		TSharedPtr<FStatusBarExtender> StatusBarExtender;
+
+		/** Cached scale for the allotted geometry; used to correctly incorporate the overall scale of the layout when drawing the image. */
+		mutable float AllottedGeometryScale = 1.0f;
+	};
+}
diff --git a/Engine/Plugins/Experimental/ImageWidgets/ImageWidgets.uplugin b/Engine/Plugins/Experimental/ImageWidgets/ImageWidgets.uplugin
deleted file mode 100644
diff --git a/Engine/Plugins/Experimental/ImageWidgets/Source/ImageWidgets/ImageWidgets.Build.cs b/Engine/Plugins/Experimental/ImageWidgets/Source/ImageWidgets/ImageWidgets.Build.cs
deleted file mode 100644
diff --git a/Engine/Plugins/Experimental/ImageWidgets/Source/ImageWidgets/Private/ColorViewerSample/ColorViewer.cpp b/Engine/Plugins/Experimental/ImageWidgets/Source/ImageWidgets/Private/ColorViewerSample/ColorViewer.cpp
deleted file mode 100644
diff --git a/Engine/Plugins/Experimental/ImageWidgets/Source/ImageWidgets/Private/ColorViewerSample/ColorViewer.h b/Engine/Plugins/Experimental/ImageWidgets/Source/ImageWidgets/Private/ColorViewerSample/ColorViewer.h
deleted file mode 100644
diff --git a/Engine/Plugins/Experimental/ImageWidgets/Source/ImageWidgets/Private/ColorViewerSample/ColorViewerCommands.cpp b/Engine/Plugins/Experimental/ImageWidgets/Source/ImageWidgets/Private/ColorViewerSample/ColorViewerCommands.cpp
deleted file mode 100644
diff --git a/Engine/Plugins/Experimental/ImageWidgets/Source/ImageWidgets/Private/ColorViewerSample/ColorViewerCommands.h b/Engine/Plugins/Experimental/ImageWidgets/Source/ImageWidgets/Private/ColorViewerSample/ColorViewerCommands.h
deleted file mode 100644
diff --git a/Engine/Plugins/Experimental/ImageWidgets/Source/ImageWidgets/Private/ColorViewerSample/ColorViewerStyle.cpp b/Engine/Plugins/Experimental/ImageWidgets/Source/ImageWidgets/Private/ColorViewerSample/ColorViewerStyle.cpp
deleted file mode 100644
diff --git a/Engine/Plugins/Experimental/ImageWidgets/Source/ImageWidgets/Private/ColorViewerSample/ColorViewerStyle.h b/Engine/Plugins/Experimental/ImageWidgets/Source/ImageWidgets/Private/ColorViewerSample/ColorViewerStyle.h
deleted file mode 100644
diff --git a/Engine/Plugins/Experimental/ImageWidgets/Source/ImageWidgets/Private/ColorViewerSample/ColorViewerWidget.cpp b/Engine/Plugins/Experimental/ImageWidgets/Source/ImageWidgets/Private/ColorViewerSample/ColorViewerWidget.cpp
deleted file mode 100644
diff --git a/Engine/Plugins/Experimental/ImageWidgets/Source/ImageWidgets/Private/ColorViewerSample/ColorViewerWidget.h b/Engine/Plugins/Experimental/ImageWidgets/Source/ImageWidgets/Private/ColorViewerSample/ColorViewerWidget.h
deleted file mode 100644
diff --git a/Engine/Plugins/Experimental/ImageWidgets/Source/ImageWidgets/Private/ImageABComparison.h b/Engine/Plugins/Experimental/ImageWidgets/Source/ImageWidgets/Private/ImageABComparison.h
deleted file mode 100644
diff --git a/Engine/Plugins/Experimental/ImageWidgets/Source/ImageWidgets/Private/ImageViewportClient.cpp b/Engine/Plugins/Experimental/ImageWidgets/Source/ImageWidgets/Private/ImageViewportClient.cpp
deleted file mode 100644
diff --git a/Engine/Plugins/Experimental/ImageWidgets/Source/ImageWidgets/Private/ImageViewportClient.h b/Engine/Plugins/Experimental/ImageWidgets/Source/ImageWidgets/Private/ImageViewportClient.h
deleted file mode 100644
diff --git a/Engine/Plugins/Experimental/ImageWidgets/Source/ImageWidgets/Private/ImageViewportController.cpp b/Engine/Plugins/Experimental/ImageWidgets/Source/ImageWidgets/Private/ImageViewportController.cpp
deleted file mode 100644
diff --git a/Engine/Plugins/Experimental/ImageWidgets/Source/ImageWidgets/Private/ImageViewportController.h b/Engine/Plugins/Experimental/ImageWidgets/Source/ImageWidgets/Private/ImageViewportController.h
deleted file mode 100644
diff --git a/Engine/Plugins/Experimental/ImageWidgets/Source/ImageWidgets/Private/ImageWidgetsCommands.cpp b/Engine/Plugins/Experimental/ImageWidgets/Source/ImageWidgets/Private/ImageWidgetsCommands.cpp
deleted file mode 100644
diff --git a/Engine/Plugins/Experimental/ImageWidgets/Source/ImageWidgets/Private/ImageWidgetsCommands.h b/Engine/Plugins/Experimental/ImageWidgets/Source/ImageWidgets/Private/ImageWidgetsCommands.h
deleted file mode 100644
diff --git a/Engine/Plugins/Experimental/ImageWidgets/Source/ImageWidgets/Private/ImageWidgetsLogCategory.cpp b/Engine/Plugins/Experimental/ImageWidgets/Source/ImageWidgets/Private/ImageWidgetsLogCategory.cpp
deleted file mode 100644
diff --git a/Engine/Plugins/Experimental/ImageWidgets/Source/ImageWidgets/Private/ImageWidgetsLogCategory.h b/Engine/Plugins/Experimental/ImageWidgets/Source/ImageWidgets/Private/ImageWidgetsLogCategory.h
deleted file mode 100644
diff --git a/Engine/Plugins/Experimental/ImageWidgets/Source/ImageWidgets/Private/ImageWidgetsModule.cpp b/Engine/Plugins/Experimental/ImageWidgets/Source/ImageWidgets/Private/ImageWidgetsModule.cpp
deleted file mode 100644
diff --git a/Engine/Plugins/Experimental/ImageWidgets/Source/ImageWidgets/Private/ImageWidgetsStyle.cpp b/Engine/Plugins/Experimental/ImageWidgets/Source/ImageWidgets/Private/ImageWidgetsStyle.cpp
deleted file mode 100644
diff --git a/Engine/Plugins/Experimental/ImageWidgets/Source/ImageWidgets/Private/ImageWidgetsStyle.h b/Engine/Plugins/Experimental/ImageWidgets/Source/ImageWidgets/Private/ImageWidgetsStyle.h
deleted file mode 100644
diff --git a/Engine/Plugins/Experimental/ImageWidgets/Source/ImageWidgets/Private/SImageCatalog.cpp b/Engine/Plugins/Experimental/ImageWidgets/Source/ImageWidgets/Private/SImageCatalog.cpp
deleted file mode 100644
diff --git a/Engine/Plugins/Experimental/ImageWidgets/Source/ImageWidgets/Private/SImageCatalogItem.cpp b/Engine/Plugins/Experimental/ImageWidgets/Source/ImageWidgets/Private/SImageCatalogItem.cpp
deleted file mode 100644
diff --git a/Engine/Plugins/Experimental/ImageWidgets/Source/ImageWidgets/Private/SImageCatalogItem.h b/Engine/Plugins/Experimental/ImageWidgets/Source/ImageWidgets/Private/SImageCatalogItem.h
deleted file mode 100644
diff --git a/Engine/Plugins/Experimental/ImageWidgets/Source/ImageWidgets/Private/SImageViewport.cpp b/Engine/Plugins/Experimental/ImageWidgets/Source/ImageWidgets/Private/SImageViewport.cpp
deleted file mode 100644
diff --git a/Engine/Plugins/Experimental/ImageWidgets/Source/ImageWidgets/Private/SImageViewportToolbar.cpp b/Engine/Plugins/Experimental/ImageWidgets/Source/ImageWidgets/Private/SImageViewportToolbar.cpp
deleted file mode 100644
diff --git a/Engine/Plugins/Experimental/ImageWidgets/Source/ImageWidgets/Private/SImageViewportToolbar.h b/Engine/Plugins/Experimental/ImageWidgets/Source/ImageWidgets/Private/SImageViewportToolbar.h
deleted file mode 100644
diff --git a/Engine/Plugins/Experimental/ImageWidgets/Source/ImageWidgets/Public/IImageViewer.h b/Engine/Plugins/Experimental/ImageWidgets/Source/ImageWidgets/Public/IImageViewer.h
deleted file mode 100644
diff --git a/Engine/Plugins/Experimental/ImageWidgets/Source/ImageWidgets/Public/SImageCatalog.h b/Engine/Plugins/Experimental/ImageWidgets/Source/ImageWidgets/Public/SImageCatalog.h
deleted file mode 100644
diff --git a/Engine/Plugins/Experimental/ImageWidgets/Source/ImageWidgets/Public/SImageViewport.h b/Engine/Plugins/Experimental/ImageWidgets/Source/ImageWidgets/Public/SImageViewport.h
deleted file mode 100644
