Fix crash opening TextureGraph: initialize _ToolbarExtenderStyle in SImageViewport

Initialize _ToolbarExtenderStyle(nullptr) in FArguments constructor.
Raw pointer was left indeterminate, causing a garbage ISlateStyle* to be
used instead of the FAppStyle fallback, crashing in BuildMultiBoxWidget.

Fixes Case 37130.

#Synced-CL 52059773

diff --git a/Engine/Plugins/Editor/ImageWidgets/Source/ImageWidgets/Public/SImageViewport.h b/Engine/Plugins/Editor/ImageWidgets/Source/ImageWidgets/Public/SImageViewport.h
--- a/Engine/Plugins/Editor/ImageWidgets/Source/ImageWidgets/Public/SImageViewport.h
+++ b/Engine/Plugins/Editor/ImageWidgets/Source/ImageWidgets/Public/SImageViewport.h
@@ -149,7 +149,8 @@ namespace UE::ImageWidgets
 		};
 
 		SLATE_BEGIN_ARGS(SImageViewport)
-			: _bABComparisonEnabled(false)
+			: _ToolbarExtenderStyle(nullptr)
+			, _bABComparisonEnabled(false)
 			{
 			}
 
